OSC Plugin Refactor:
- Code standards clean-up
- Fix hang when socket fails attempting to send packets
- Fix unbounded BP API (resulting in crashes)
- Make OSCAddress first-class type
    - Allows us to do more specific string parsing things with OSCAddresses as part of the OSC spec (TODO: matching with their light filtering spec)
- More convenience functions for BP API
- More logging/warning/errors when attempting to stand-up server/client

diff --git a/Engine/Plugins/Runtime/OSC/Source/OSC/Classes/OSCBundle.h b/Engine/Plugins/Runtime/OSC/Source/OSC/Classes/OSCBundle.h
deleted file mode 100644
diff --git a/Engine/Plugins/Runtime/OSC/Source/OSC/Classes/OSCClient.h b/Engine/Plugins/Runtime/OSC/Source/OSC/Classes/OSCClient.h
deleted file mode 100644
diff --git a/Engine/Plugins/Runtime/OSC/Source/OSC/Classes/OSCLog.h b/Engine/Plugins/Runtime/OSC/Source/OSC/Classes/OSCLog.h
deleted file mode 100644
diff --git a/Engine/Plugins/Runtime/OSC/Source/OSC/Classes/OSCMessage.h b/Engine/Plugins/Runtime/OSC/Source/OSC/Classes/OSCMessage.h
deleted file mode 100644
diff --git a/Engine/Plugins/Runtime/OSC/Source/OSC/Classes/OSCPacket.h b/Engine/Plugins/Runtime/OSC/Source/OSC/Classes/OSCPacket.h
deleted file mode 100644
diff --git a/Engine/Plugins/Runtime/OSC/Source/OSC/Classes/OSCServer.h b/Engine/Plugins/Runtime/OSC/Source/OSC/Classes/OSCServer.h
deleted file mode 100644
diff --git a/Engine/Plugins/Runtime/OSC/Source/OSC/Classes/OSCStream.h b/Engine/Plugins/Runtime/OSC/Source/OSC/Classes/OSCStream.h
deleted file mode 100644
diff --git a/Engine/Plugins/Runtime/OSC/Source/OSC/Classes/OSCTypes.h b/Engine/Plugins/Runtime/OSC/Source/OSC/Classes/OSCTypes.h
deleted file mode 100644
diff --git a/Engine/Plugins/Runtime/OSC/Source/OSC/OSC.Build.cs b/Engine/Plugins/Runtime/OSC/Source/OSC/OSC.Build.cs
--- a/Engine/Plugins/Runtime/OSC/Source/OSC/OSC.Build.cs
+++ b/Engine/Plugins/Runtime/OSC/Source/OSC/OSC.Build.cs
@@ -6,6 +6,8 @@ namespace UnrealBuildTool.Rules
 	{
         public OSC(ReadOnlyTargetRules Target) : base(Target)
 		{
+            OptimizeCode = CodeOptimization.Never;
+            
             PublicDependencyModuleNames.AddRange(
 				new string[] {
                     "Core",
diff --git a/Engine/Plugins/Runtime/OSC/Source/OSC/Private/OSCBundle.cpp b/Engine/Plugins/Runtime/OSC/Source/OSC/Private/OSCBundle.cpp
--- a/Engine/Plugins/Runtime/OSC/Source/OSC/Private/OSCBundle.cpp
+++ b/Engine/Plugins/Runtime/OSC/Source/OSC/Private/OSCBundle.cpp
@@ -1,13 +1,18 @@
 // Copyright 1998-2019 Epic Games, Inc. All Rights Reserved.
 #include "OSCBundle.h"
+#include "OSCLog.h"
 #include "OSCStream.h"
 
-FOSCBundlePacket::FOSCBundlePacket() 
-: FOSCPacket(), TimeTag(0) 
-{}
 
-FOSCBundlePacket::~FOSCBundlePacket() 
-{}
+FOSCBundlePacket::FOSCBundlePacket()
+	: FOSCPacket()
+	, TimeTag(0)
+{
+}
+
+FOSCBundlePacket::~FOSCBundlePacket()
+{
+}
 
 void FOSCBundlePacket::SetTimeTag(uint64 NewTimeTag) 
 { 
@@ -19,65 +24,66 @@ uint64 FOSCBundlePacket::GetTimeTag() const
 	return TimeTag.GetTimeTag(); 
 }
 
-void FOSCBundlePacket::AddPacket(TSharedPtr<FOSCPacket> Packet)
-{
-	Packets.Add(Packet);
-}
-
 void FOSCBundlePacket::WriteData(FOSCStream& Stream)
 {
-	// Write bundle
-	Stream.WriteString(FString("#bundle"));
+	if (!GetAddress().IsBundle())
+	{
+		UE_LOG(LogOSC, Warning, TEXT("Failed to write OSCBundlePacket. Invalid OSCAddress '%s'"), *GetAddress().Value);
+		return;
+	}
 
-	// Write bundle time tag
+	// Write bundle & time tag
+	Stream.WriteString(GetAddress().Value);
 	Stream.WriteUInt64(GetTimeTag());
 
-	// for each packet
-	for (auto& packet : Packets)
+	for (TSharedPtr<FOSCPacket>& Packet : Packets)
 	{
-		int sizePosition = Stream.GetPosition();
+		int32 StreamPos = Stream.GetPosition();
 		Stream.WriteInt32(0);
 
-		int prePos = Stream.GetPosition();
-		packet->WriteData(Stream);
-		int postPos = Stream.GetPosition();
+		int32 InitPos = Stream.GetPosition();
+		Packet->WriteData(Stream);
+		int32 NewPos = Stream.GetPosition();
 
-		Stream.SetPosition(sizePosition);
-		Stream.WriteInt32(postPos - prePos);
-		Stream.SetPosition(postPos);
+		Stream.SetPosition(StreamPos);
+		Stream.WriteInt32(NewPos - InitPos);
+		Stream.SetPosition(NewPos);
 	}
 }
 
 void FOSCBundlePacket::ReadData(FOSCStream& Stream)
 {
-	// Read address
-	FString address = Stream.ReadString();
-
-	// read time tag
 	TimeTag = FOSCType(Stream.ReadUInt64());
 
-	// While we still have data
 	while (!Stream.HasReachedEnd())
 	{
-		int32 sizeOfPacket = Stream.ReadInt32();
-
-		TSharedPtr<FOSCPacket> packet = FOSCPacket::CreatePacket(Stream.GetBuffer());
+		TSharedPtr<FOSCPacket> Packet = FOSCPacket::CreatePacket(Stream.GetData());
 
-		if (packet.IsValid())
+		if (Packet.IsValid())
 		{
-			packet->ReadData(Stream);
-			Packets.Add(packet);
-		}		
+			Packet->ReadData(Stream);
+			Packets.Add(Packet);
+		}
 	}
 }
 
-const TSharedPtr<FOSCBundlePacket> FOSCBundle::GetOrCreatePacket()
+FOSCBundlePacket::FPacketBundle& FOSCBundlePacket::GetPackets()
 {
-	if (!Packet.IsValid())
-	{
-		Packet = MakeShareable(new FOSCBundlePacket());
-	}
+	return Packets;
+}
 
-	return Packet;
+bool FOSCBundlePacket::IsBundle()
+{
+	return true;
 }
 
+bool FOSCBundlePacket::IsMessage()
+{
+	return false;
+}
+
+const FOSCAddress& FOSCBundlePacket::GetAddress() const
+{
+	const static FOSCAddress BundleIdentifier = TEXT("#bundle");
+	return BundleIdentifier;
+}
\ No newline at end of file
diff --git a/Engine/Plugins/Runtime/OSC/Source/OSC/Private/OSCClient.cpp b/Engine/Plugins/Runtime/OSC/Source/OSC/Private/OSCClient.cpp
--- a/Engine/Plugins/Runtime/OSC/Source/OSC/Private/OSCClient.cpp
+++ b/Engine/Plugins/Runtime/OSC/Source/OSC/Private/OSCClient.cpp
@@ -8,40 +8,83 @@
 #include "OSCBundle.h"
 #include "OSCLog.h"
 
+
+namespace
+{
+	static const int32 OUTPUT_BUFFER_SIZE = 1024;
+} // namespace <>
+
 UOSCClient::UOSCClient()
-	: Socket(FUdpSocketBuilder(TEXT("OscSender")).Build())
+	: Socket(FUdpSocketBuilder(*GetName()).Build())
+{
+}
+
+UOSCClient::~UOSCClient()
 {
-	//
 }
 
-void UOSCClient::SetTarget(FIPv4Address InIPAddress, uint32_t Port)
+bool UOSCClient::SetSendIPAddress(const FString& InIPAddress, const int32 Port)
 {
-	this->IPAddress = ISocketSubsystem::Get(PLATFORM_SOCKETSUBSYSTEM)->CreateInternetAddr();
+	IPAddress = ISocketSubsystem::Get(PLATFORM_SOCKETSUBSYSTEM)->CreateInternetAddr();
 
-	this->IPAddress->SetIp(InIPAddress.Value);
-	this->IPAddress->SetPort(Port);
+	bool bIsValidAddress = true;
+	IPAddress->SetIp(*InIPAddress, bIsValidAddress);
+	IPAddress->SetPort(Port);
 
-	UE_LOG(OSCLog, Display, TEXT("Set target port: %d"), Port);
+	if (bIsValidAddress)
+	{
+		UE_LOG(LogOSC, Verbose, TEXT("OSCClient '%s' SetSendIpAddress: %s:%d"), *GetName(), *InIPAddress, Port);
+	}
+	else
+	{
+		UE_LOG(LogOSC, Warning, TEXT("OSCClient '%s' SetSendIpAddress Failed for input: %s:%d"), *GetName(), *InIPAddress, Port);
+	}
+
+	return bIsValidAddress;
 }
 
-void UOSCClient::SendPacket(FOSCPacket* Packet, char* Data, int32 Size)
+void UOSCClient::GetSendIPAddress(FString& InIPAddress, int32& Port)
 {
-	if (IPAddress)
+	const bool bAppendPort = false;
+	InIPAddress = IPAddress->ToString(bAppendPort);
+	Port = IPAddress->GetPort();
+}
+
+void UOSCClient::SendPacket(FOSCPacket& Packet)
+{
+	const FOSCAddress& OSCAddress = Packet.GetAddress();
+	if (!OSCAddress.IsValid())
 	{
-		FOSCStream stream = FOSCStream(Data, Size);
+		UE_LOG(LogOSC, Warning, TEXT("Failed to write bundle data. Invalid OSCAddress '%s'"), *OSCAddress.Value);
+		return;
+	}
 
-		Packet->WriteData(stream);
+	if (IPAddress)
+	{
+		FOSCStream Stream = FOSCStream(OUTPUT_BUFFER_SIZE);
+		const uint8* DataPtr = Stream.GetData();
+		Packet.WriteData(Stream);
 
-		int32 length = stream.GetPosition();
+		int32 BytesSent = 0;
 
-		int32 bytesSent = 0;
-		while (length > 0)
+		const int32 AttemptedLength = Stream.GetPosition();
+		int32 Length = AttemptedLength;
+		while (Length > 0)
 		{
-			Socket->SendTo((uint8*)Data, length, bytesSent, *IPAddress);
-
-			length -= bytesSent;
-			Data += bytesSent;
+			const bool bSuccess = Socket->SendTo(DataPtr, Length, BytesSent, *IPAddress);
+			if (!bSuccess || BytesSent <= 0)
+			{
+				UE_LOG(LogOSC, Verbose, TEXT("OSC Packet failed: Client '%s', OSCAddress '%s', Send IPAddress %s, Attempted Bytes = %d"),
+					*GetName(), *OSCAddress.Value, *IPAddress->ToString(true /*bAppendPort*/), AttemptedLength);
+				return;
+			}
+
+			Length -= BytesSent;
+			DataPtr += BytesSent;
 		}
+
+		UE_LOG(LogOSC, Verbose, TEXT("OSC Packet sent: Client '%s', OSCAddress '%s', Send IPAddress %s, Bytes Sent = %d"),
+			*GetName(), *OSCAddress.Value, *IPAddress->ToString(true /*bAppendPort*/), AttemptedLength);
 	}
 }
 
@@ -60,20 +103,16 @@ void UOSCClient::BeginDestroy()
 	Super::BeginDestroy();
 }
 
-void UOSCClient::SendOSCMessage(FString Address, UPARAM(ref) FOSCMessage& Msg)
+void UOSCClient::SendOSCMessage(FOSCMessage& Message)
 {
-	static const int OUTPUT_BUFFER_SIZE = 1024;
-	char buffer[OUTPUT_BUFFER_SIZE];
-
-	Msg.GetPacket()->SetAddress(Address);
-
-	SendPacket(&(*Msg.GetPacket()), &buffer[0], OUTPUT_BUFFER_SIZE);
+	const TSharedPtr<FOSCMessagePacket>& Packet = Message.GetPacket();
+	check(Packet.IsValid());
+	SendPacket(*Packet.Get());
 }
 
-void UOSCClient::SendOSCBundle(UPARAM(ref) FOSCBundle& Bundle)
+void UOSCClient::SendOSCBundle(FOSCBundle& Bundle)
 {
-	static const int OUTPUT_BUFFER_SIZE = 1024;
-	char buffer[OUTPUT_BUFFER_SIZE];
-
-	SendPacket(&(*Bundle.GetPacket()), &buffer[0], OUTPUT_BUFFER_SIZE);
+	const TSharedPtr<FOSCBundlePacket>& Packet = Bundle.GetPacket();
+	check(Packet.IsValid());
+	SendPacket(*Packet.Get());
 }
diff --git a/Engine/Plugins/Runtime/OSC/Source/OSC/Private/OSCMessage.cpp b/Engine/Plugins/Runtime/OSC/Source/OSC/Private/OSCMessage.cpp
--- a/Engine/Plugins/Runtime/OSC/Source/OSC/Private/OSCMessage.cpp
+++ b/Engine/Plugins/Runtime/OSC/Source/OSC/Private/OSCMessage.cpp
@@ -2,146 +2,48 @@
 #include "OSCMessage.h"
 #include "OSCStream.h"
 
-FOSCMessagePacket::FOSCMessagePacket() 
-	: FOSCPacket() 
-{}
-
-FOSCMessagePacket::~FOSCMessagePacket() 
-{}
-
-
-void FOSCMessagePacket::WriteData(FOSCStream& Stream)
+FOSCMessage::FOSCMessage()
+	: Packet(MakeShareable(new FOSCMessagePacket()))
 {
-	// Begin writing data 
-	Stream.WriteString(Address);
+}
 
-	// write type tags
-	FString tagTypes = ",";
-	for (int i = 0; i < GetNumArguments(); i++)
-	{
-		tagTypes += static_cast<char>(GetArgument(i).GetTypeTag());
-	}
+FOSCMessage::FOSCMessage(const TSharedPtr<FOSCMessagePacket>& InPacket)
+	: Packet(InPacket)
+{
+}
 
-	Stream.WriteString(tagTypes);
+FOSCMessage::~FOSCMessage()
+{
+	Packet.Reset();
+}
 
-	for (FOSCType type : Arguments)
-	{
-		switch (type.GetTypeTag())
-		{
-		case EOSCTypeTag::CHAR:
-			Stream.WriteChar(type.GetChar());
-			break;
-		case EOSCTypeTag::INT32:
-			Stream.WriteInt32(type.GetInt32());
-			break;
-		case EOSCTypeTag::FLOAT:
-			Stream.WriteFloat(type.GetFloat());
-			break;
-		case EOSCTypeTag::DOUBLE:
-			Stream.WriteDouble(type.GetDouble());
-			break;
-		case EOSCTypeTag::INT64:
-			Stream.WriteInt64(type.GetInt64());
-			break;
-		case EOSCTypeTag::TIME:
-			Stream.WriteUInt64(type.GetTimeTag());
-			break;
-		case EOSCTypeTag::STRING:
-			Stream.WriteString(type.GetString());
-			break;
-		case EOSCTypeTag::BLOB:
-		{
-			TArray<uint8> blob = type.GetBlob();
-			Stream.WriteBlob(blob);
-		}
-		break;
-		case EOSCTypeTag::COLOR:
-			Stream.WriteInt32(type.GetColor().GetInt32());
-			break;
-		case EOSCTypeTag::TRUE:
-		case EOSCTypeTag::FALSE:
-		case EOSCTypeTag::NIL:
-		case EOSCTypeTag::INFINITUM:
-			// No values are written for these types
-			break;
-		default:
-			// Argument is not supported 
-			unimplemented();
-			break;
-		}
-	}
+void FOSCMessage::SetPacket(TSharedPtr<FOSCMessagePacket>& InPacket)
+{
+	Packet = InPacket;
 }
 
-void FOSCMessagePacket::ReadData(FOSCStream& Stream)
+const TSharedPtr<FOSCMessagePacket>& FOSCMessage::GetPacket() const
 {
-	// Read Address
-	Address = Stream.ReadString();
+	return Packet;
+}
 
-	// Read string of tags
-	TArray<TCHAR> tagTypes = Stream.ReadString().GetCharArray();
 
-	// Skip the first argument which is ','
-	for (int i=1; i<tagTypes.Num(); i++)
-	{
-		switch ((EOSCTypeTag)tagTypes[i])
-		{
-		case EOSCTypeTag::CHAR:
-			AddArgument(FOSCType(Stream.ReadChar()));
-			break;
-		case EOSCTypeTag::INT32:
-			AddArgument(FOSCType(Stream.ReadInt32()));
-			break;
-		case EOSCTypeTag::FLOAT:
-			AddArgument(FOSCType(Stream.ReadFloat()));
-			break;
-		case EOSCTypeTag::DOUBLE:
-			AddArgument(FOSCType(Stream.ReadDouble()));
-			break;
-		case EOSCTypeTag::INT64:
-			AddArgument(FOSCType(Stream.ReadInt64()));
-			break;
-		case EOSCTypeTag::TRUE:
-			AddArgument(FOSCType(true));
-			break;
-		case EOSCTypeTag::FALSE:
-			AddArgument(FOSCType(false));
-			break;
-		case EOSCTypeTag::NIL:
-			AddArgument(FOSCType(EOSCTypeTag::NIL));
-			break;
-		case EOSCTypeTag::INFINITUM:
-			AddArgument(FOSCType(EOSCTypeTag::INFINITUM));
-			break;
-		case EOSCTypeTag::TIME:
-			AddArgument(FOSCType(Stream.ReadUInt64()));
-			break;
-		case EOSCTypeTag::STRING:
-			AddArgument(FOSCType(Stream.ReadString()));
-			break;
-		case EOSCTypeTag::BLOB:
-			AddArgument(FOSCType(Stream.ReadBlob()));
-			break;
-		case EOSCTypeTag::COLOR:
-			AddArgument(FOSCType(FOSCColor(Stream.ReadInt32())));
-			break;
-		case EOSCTypeTag::TERMINATE:
-			Stream.ReadChar();
-			break;
+bool FOSCMessage::SetAddress(const FOSCAddress& InAddress)
+{
+	check(Packet.IsValid());
 
-		default:
-			// Argument is not supported 
-			unimplemented();
-			break;
-		}
+	if (!InAddress.IsMessage())
+	{
+		UE_LOG(LogOSC, Warning, TEXT("Attempting to set invalid OSCAddress '%s'. OSC address must begin with '/'"), *InAddress.Value);
+		return false;
 	}
+
+	Packet->SetAddress(InAddress);
+	return true;
 }
 
-const TSharedPtr<FOSCMessagePacket> FOSCMessage::GetOrCreatePacket()
+const FOSCAddress& FOSCMessage::GetAddress() const
 {
-	if (!Packet.IsValid())
-	{
-		Packet = MakeShareable(new FOSCMessagePacket());
-	}
-
-	return Packet;
+	check(Packet.IsValid());
+	return Packet->GetAddress();
 }
\ No newline at end of file
diff --git a/Engine/Plugins/Runtime/OSC/Source/OSC/Private/OSCPacket.cpp b/Engine/Plugins/Runtime/OSC/Source/OSC/Private/OSCPacket.cpp
--- a/Engine/Plugins/Runtime/OSC/Source/OSC/Private/OSCPacket.cpp
+++ b/Engine/Plugins/Runtime/OSC/Source/OSC/Private/OSCPacket.cpp
@@ -1,18 +1,200 @@
 // Copyright 1998-2019 Epic Games, Inc. All Rights Reserved.
 #include "OSCPacket.h"
+
 #include "OSCMessage.h"
 #include "OSCBundle.h"
+#include "OSCLog.h"
+
+
+FOSCPacket::FOSCPacket()
+{
+}
+
+FOSCPacket::~FOSCPacket()
+{
+}
+
+TSharedPtr<FOSCPacket> FOSCPacket::CreatePacket(const uint8* PacketType)
+{
+	const FOSCAddress Address(ANSI_TO_TCHAR((const ANSICHAR*)&PacketType[0]));
+	if (Address.IsMessage())
+	{
+		return MakeShareable(new FOSCMessagePacket());
+	}
+	else if (Address.IsBundle())
+	{
+		return MakeShareable(new FOSCBundlePacket());
+	}
+	else
+	{
+		UE_LOG(LogOSC, Warning, TEXT("Failed to parse lead character of OSC message packet. "
+			"Lead identifier of '%c' not valid bundle (#) or message (/) identifier."), PacketType[0]);
+		return nullptr;
+	}
+}
+
+FOSCMessagePacket::FOSCMessagePacket()
+	: FOSCPacket()
+{}
+
+FOSCMessagePacket::~FOSCMessagePacket()
+{}
+
+void FOSCMessagePacket::SetAddress(const FOSCAddress& InAddress)
+{
+	Address = InAddress;
+}
+
+const FOSCAddress& FOSCMessagePacket::GetAddress() const
+{
+	return Address;
+}
+
+TArray<FOSCType>& FOSCMessagePacket::GetArguments()
+{
+	return Arguments;
+}
+
+bool FOSCMessagePacket::IsBundle()
+{
+	return false;
+}
 
-TSharedPtr<FOSCPacket> FOSCPacket::CreatePacket(const char* PacketType)
+bool FOSCMessagePacket::IsMessage()
 {
-	if (PacketType[0] == '/')
+	return true;
+}
+
+void FOSCMessagePacket::WriteData(FOSCStream& Stream)
+{
+	if (!Address.IsMessage())
 	{
-		return  MakeShareable(new FOSCMessagePacket());
+		UE_LOG(LogOSC, Warning, TEXT("Failed to write OSCMessagePacket. Invalid OSCAddress '%s'"), *Address.Value);
+		return;
 	}
-	else if (PacketType[0] == '#')
+
+	// Begin writing data 
+	Stream.WriteString(Address.Value);
+
+	// write type tags
+	FString TagTypes = ",";
+	for (int32 i = 0; i < Arguments.Num(); i++)
 	{
-		return  MakeShareable(new FOSCBundlePacket());
+		TagTypes += static_cast<ANSICHAR>(Arguments[i].GetTypeTag());
 	}
 
-	return nullptr;
-}
\ No newline at end of file
+	Stream.WriteString(TagTypes);
+
+	for (FOSCType OSCType : Arguments)
+	{
+		switch (OSCType.GetTypeTag())
+		{
+		case EOSCTypeTag::CHAR:
+			Stream.WriteChar(OSCType.GetChar());
+			break;
+		case EOSCTypeTag::INT32:
+			Stream.WriteInt32(OSCType.GetInt32());
+			break;
+		case EOSCTypeTag::FLOAT:
+			Stream.WriteFloat(OSCType.GetFloat());
+			break;
+		case EOSCTypeTag::DOUBLE:
+			Stream.WriteDouble(OSCType.GetDouble());
+			break;
+		case EOSCTypeTag::INT64:
+			Stream.WriteInt64(OSCType.GetInt64());
+			break;
+		case EOSCTypeTag::TIME:
+			Stream.WriteUInt64(OSCType.GetTimeTag());
+			break;
+		case EOSCTypeTag::STRING:
+			Stream.WriteString(OSCType.GetString());
+			break;
+		case EOSCTypeTag::BLOB:
+		{
+			TArray<uint8> blob = OSCType.GetBlob();
+			Stream.WriteBlob(blob);
+		}
+		break;
+		case EOSCTypeTag::COLOR:
+			Stream.WriteInt32(OSCType.GetColor().GetInt32());
+			break;
+		case EOSCTypeTag::TRUE:
+		case EOSCTypeTag::FALSE:
+		case EOSCTypeTag::NIL:
+		case EOSCTypeTag::INFINITUM:
+			// No values are written for these types
+			break;
+		default:
+			// Argument is not supported 
+			unimplemented();
+			break;
+		}
+	}
+}
+
+void FOSCMessagePacket::ReadData(FOSCStream& Stream)
+{
+	// Read Address
+	Address = Stream.ReadString();
+
+	// Read string of tags
+	const FString StreamString = Stream.ReadString();
+	const TArray<TCHAR>& TagTypes = StreamString.GetCharArray();
+
+	// Skip the first argument which is ','
+	for (int32 i = 1; i < TagTypes.Num(); i++)
+	{
+		const EOSCTypeTag Tag = static_cast<EOSCTypeTag>(TagTypes[i]);
+		switch (Tag)
+		{
+		case EOSCTypeTag::CHAR:
+			Arguments.Add(FOSCType(Stream.ReadChar()));
+			break;
+		case EOSCTypeTag::INT32:
+			Arguments.Add(FOSCType(Stream.ReadInt32()));
+			break;
+		case EOSCTypeTag::FLOAT:
+			Arguments.Add(FOSCType(Stream.ReadFloat()));
+			break;
+		case EOSCTypeTag::DOUBLE:
+			Arguments.Add(FOSCType(Stream.ReadDouble()));
+			break;
+		case EOSCTypeTag::INT64:
+			Arguments.Add(FOSCType(Stream.ReadInt64()));
+			break;
+		case EOSCTypeTag::TRUE:
+			Arguments.Add(FOSCType(true));
+			break;
+		case EOSCTypeTag::FALSE:
+			Arguments.Add(FOSCType(false));
+			break;
+		case EOSCTypeTag::NIL:
+			Arguments.Add(FOSCType(EOSCTypeTag::NIL));
+			break;
+		case EOSCTypeTag::INFINITUM:
+			Arguments.Add(FOSCType(EOSCTypeTag::INFINITUM));
+			break;
+		case EOSCTypeTag::TIME:
+			Arguments.Add(FOSCType(Stream.ReadUInt64()));
+			break;
+		case EOSCTypeTag::STRING:
+			Arguments.Add(FOSCType(Stream.ReadString()));
+			break;
+		case EOSCTypeTag::BLOB:
+			Arguments.Add(FOSCType(Stream.ReadBlob()));
+			break;
+		case EOSCTypeTag::COLOR:
+			Arguments.Add(FOSCType(FOSCColor(Stream.ReadInt32())));
+			break;
+		case EOSCTypeTag::TERMINATE:
+			Stream.ReadChar();
+			break;
+
+		default:
+			// Argument is not supported 
+			unimplemented();
+			break;
+		}
+	}
+}
diff --git a/Engine/Plugins/Runtime/OSC/Source/OSC/Private/OSCPacket.h b/Engine/Plugins/Runtime/OSC/Source/OSC/Private/OSCPacket.h
new file mode 100644
--- /dev/null
+++ b/Engine/Plugins/Runtime/OSC/Source/OSC/Private/OSCPacket.h
@@ -0,0 +1,107 @@
+// Copyright 1998-2019 Epic Games, Inc. All Rights Reserved.
+#pragma once
+
+#include "OSCTypes.h"
+#include "OSCStream.h"
+#include "OSCAddress.h"
+
+// Forward Declarations
+class FOSCStream;
+
+
+class FOSCPacket
+{
+public:
+	FOSCPacket();
+	virtual ~FOSCPacket();
+
+	/** Write packet data into stream */
+	virtual void WriteData(FOSCStream& Stream) = 0;
+	
+	/** Read packet data from stream */
+	virtual void ReadData(FOSCStream& Stream) = 0;
+	
+	/** Returns OSC address identifier of packet */
+	virtual const FOSCAddress& GetAddress() const = 0;
+
+	/** Returns true if packet is message */
+	virtual bool IsMessage() = 0;
+	
+	/** Returns true if packet is bundle */
+	virtual bool IsBundle() = 0;
+
+	/** Create an OSC packet according to the input data. */
+	static TSharedPtr<FOSCPacket> CreatePacket(const uint8* PacketType);
+};
+
+class FOSCBundlePacket : public FOSCPacket
+{
+public:
+	FOSCBundlePacket();
+	virtual ~FOSCBundlePacket();
+
+	using FPacketBundle = TArray<TSharedPtr<FOSCPacket>>;
+
+	/** Set the bundle time tag. */
+	void SetTimeTag(uint64 NewTimeTag);
+
+	/** Get the bundle time tag. */
+	uint64 GetTimeTag() const;
+
+	/** Get OSC packet by index. */
+	FPacketBundle& GetPackets();
+
+	virtual bool IsBundle() override;
+	virtual bool IsMessage() override;
+
+	/** Writes bundle data into the OSC stream. */
+	virtual void WriteData(FOSCStream& Stream) override;
+
+	/** Reads bundle data from provided OSC stream,
+	  * adding packet data to internal packet bundle. */
+	virtual void ReadData(FOSCStream& Stream) override;
+
+	virtual const FOSCAddress& GetAddress() const override;
+
+private:
+	/** Bundle of OSC packets. */
+	FPacketBundle Packets;
+
+	/** Bundle time tag. */
+	FOSCType TimeTag;
+};
+
+class FOSCMessagePacket : public FOSCPacket
+{
+public:
+	FOSCMessagePacket();
+	virtual ~FOSCMessagePacket();
+
+	/** Set OSC message address. */
+	void SetAddress(const FOSCAddress& InAddress);
+
+	/** Get OSC message address. */
+	virtual const FOSCAddress& GetAddress() const override;
+
+	/** Get arguments array. */
+	TArray<FOSCType>& GetArguments();
+
+	/** Returns false to indicate type is not OSC bundle. */
+	virtual bool IsBundle();
+
+	/** Returns true to indicate its an OSC message. */
+	virtual bool IsMessage();
+
+	/** Write message data into an OSC stream. */
+	virtual void WriteData(FOSCStream& Stream) override;
+
+	/** Reads message data from an OSC stream and creates new argument. */
+	virtual void ReadData(FOSCStream& Stream) override;
+
+private:
+	/** OSC address. */
+	FOSCAddress Address;
+
+	/** List of argument types. */
+	TArray<FOSCType> Arguments;
+};
diff --git a/Engine/Plugins/Runtime/OSC/Source/OSC/Private/OSCServer.cpp b/Engine/Plugins/Runtime/OSC/Source/OSC/Private/OSCServer.cpp
--- a/Engine/Plugins/Runtime/OSC/Source/OSC/Private/OSCServer.cpp
+++ b/Engine/Plugins/Runtime/OSC/Source/OSC/Private/OSCServer.cpp
@@ -14,110 +14,159 @@
 #define OSC_BUFFER_FREE 0
 #define OSC_BUFFER_LOCKED 1
 
-UOSCServer::UOSCServer()
-: OSCPackets(1024)
-, OSCBufferLock(0)
+
+UOSCServer::UOSCServer(const FObjectInitializer& ObjectInitializer)
+	: Super(ObjectInitializer) 
+	, bWhitelistClients(false)
+	, OSCPackets(1024)
+	, OSCBufferLock(0)
 {
 }
 
-UOSCServer::~UOSCServer()
+void UOSCServer::Callback(const FArrayReaderPtr& Data, const FIPv4Endpoint& Endpoint)
 {
-	delete SocketReceiver;
-	SocketReceiver = nullptr;
+	// Throw request on the ground if not endpoint address not whitelisted.
+	if (bWhitelistClients && !ClientWhitelist.Contains(Endpoint.Address.Value))
+	{
+		return;
+	}
 
-	if (Socket)
+	TSharedPtr<FOSCPacket> Packet = FOSCPacket::CreatePacket(Data->GetData());
+	if (Packet.IsValid())
 	{
-		ISocketSubsystem::Get(PLATFORM_SOCKETSUBSYSTEM)->DestroySocket(Socket);
-		Socket = nullptr;
+		FOSCStream Stream = FOSCStream(Data->GetData(), Data->Num());
+
+		Packet->ReadData(Stream);
+		OSCPackets.Buffer.Enqueue(MoveTemp(Packet));
 	}
+
+	if (!FPlatformAtomics::InterlockedCompareExchange(&OSCBufferLock, OSC_BUFFER_LOCKED, OSC_BUFFER_FREE) &&
+		!OSCPackets.Buffer.IsEmpty())
+	{
+		check(OSCBufferLock == OSC_BUFFER_LOCKED);
+		FSimpleDelegateGraphTask::CreateAndDispatchWhenReady(
+			FSimpleDelegateGraphTask::FDelegate::CreateLambda([this]() {
+			check(OSCBufferLock == OSC_BUFFER_LOCKED);
+			FPlatformAtomics::InterlockedCompareExchange(&OSCBufferLock, OSC_BUFFER_FREE, OSC_BUFFER_LOCKED);
+				
+				TSharedPtr<FOSCPacket> Packet;
+				while (OSCPackets.Buffer.Dequeue(Packet))
+				{
+					if (Packet->IsMessage())
+					{
+						FOSCMessage Message(FOSCMessage(StaticCastSharedPtr<FOSCMessagePacket>(Packet)));
+						OnOscReceived.Broadcast(Message);
+					}
+					else if (Packet->IsBundle())
+					{
+						FOSCBundle Bundle(StaticCastSharedPtr<FOSCBundlePacket>(Packet));
+						OnOscBundleReceived.Broadcast(Bundle);
+					}
+					else
+					{
+						UE_LOG(LogOSC, Warning, TEXT("Failed to parse invalid received OSC message. OSCAddress '%s' is invalid."), *Packet->GetAddress().Value);
+					}
+				}
+				OSCPackets.Buffer.Empty();
+			}),
+			TStatId(),
+			nullptr,
+			ENamedThreads::GameThread);
+	}
+}
+
+bool UOSCServer::GetMulticastLoopback() const
+{
+	return bMulticastLoopback;
 }
 
-void UOSCServer::Listen(FIPv4Address IPAddress, uint32_t Port, bool MulticastLoopback)
+bool UOSCServer::IsActive() const
 {
-	FUdpSocketBuilder builder(TEXT("OscListener"));
-	builder.BoundToPort(Port);
-	if (IPAddress.IsMulticastAddress())
+	return SocketReceiver != nullptr;
+}
+
+void UOSCServer::Listen()
+{
+	if (IsActive())
+	{
+		UE_LOG(LogOSC, Error, TEXT("OSCServer currently listening: %s:%d. Failed to start new service prior to calling stop."), *GetName(), *ReceiveIPAddress.ToString(), Port);
+		return;
+	}
+
+	FUdpSocketBuilder Builder(*GetName());
+	Builder.BoundToPort(Port);
+	if (ReceiveIPAddress.IsMulticastAddress())
 	{
-		builder.JoinedToGroup(IPAddress);
-		if (MulticastLoopback)
+		Builder.JoinedToGroup(ReceiveIPAddress);
+		if (bMulticastLoopback)
 		{
-			builder.WithMulticastLoopback();
+			Builder.WithMulticastLoopback();
 		}
 	}
 	else
 	{
-		builder.BoundToAddress(IPAddress);
+		if (bMulticastLoopback)
+		{
+			UE_LOG(LogOSC, Warning, TEXT("OSCServer '%s' ReceiveIPAddress provided is not a multicast address.  Not respecting MulticastLoopback boolean."), *GetName());
+		}
+		Builder.BoundToAddress(ReceiveIPAddress);
 	}
 
-	Socket = builder.Build();
+	Socket = Builder.Build();
 	if (Socket)
 	{
-		SocketReceiver = new FUdpSocketReceiver(Socket, FTimespan::FromMilliseconds(100), TEXT("OSCListener"));
+		SocketReceiver = new FUdpSocketReceiver(Socket, FTimespan::FromMilliseconds(100), *GetName().Append(TEXT("_ListenerThread")));
 		SocketReceiver->OnDataReceived().BindUObject(this, &UOSCServer::Callback);
 		SocketReceiver->Start();
 
-		UE_LOG(OSCLog, Display, TEXT("Listening to port %d"), Port);
+		UE_LOG(LogOSC, Display, TEXT("OSCServer '%s' Listening: %s:%d."), *GetName(), *ReceiveIPAddress.ToString(), Port);
 	}
 	else
 	{
-		UE_LOG(OSCLog, Warning, TEXT("Cannot listen to port %d"), Port);
+		UE_LOG(LogOSC, Warning, TEXT("OSCServer '%s' failed to bind to socket on %s:%d."), *GetName(), *ReceiveIPAddress.ToString(), Port);
 	}
 }
 
-void UOSCServer::Callback(const FArrayReaderPtr& Data, const FIPv4Endpoint& Endpoint)
+bool UOSCServer::SetAddress(const FString& InReceiveIPAddress, int32 InPort)
 {
-	TSharedPtr <FOSCPacket> packet = FOSCPacket::CreatePacket((char *)Data->GetData());
-
-	if (packet.IsValid())
+	if (IsActive())
 	{
-		// create OSC stream
-		FOSCStream stream = FOSCStream((char *)Data->GetData(), Data->Num());
-
-		packet->ReadData(stream);
+		UE_LOG(LogOSC, Error, TEXT("Cannot set address while OSCServer is active."));
+		return false;
+	}
 
-		OSCPackets.Buffer.Enqueue(MoveTemp(packet));
+	if (!FIPv4Address::Parse(InReceiveIPAddress, ReceiveIPAddress))
+	{
+		UE_LOG(LogOSC, Error, TEXT("Invalid ReceiveIPAddress when attempting to set on OSCServer '%s'. Address not updated."), *GetName());
+		return false;
 	}
 
-	if (!FPlatformAtomics::InterlockedCompareExchange(&OSCBufferLock, OSC_BUFFER_LOCKED, OSC_BUFFER_FREE) &&
-		!OSCPackets.Buffer.IsEmpty())
+	Port = InPort;
+	return true;
+}
+
+void UOSCServer::SetMulticastLoopback(bool bInMulticastLoopback)
+{
+	if (bInMulticastLoopback != bMulticastLoopback && IsActive())
 	{
-		check(OSCBufferLock == OSC_BUFFER_LOCKED);
-		FSimpleDelegateGraphTask::CreateAndDispatchWhenReady(
-			FSimpleDelegateGraphTask::FDelegate::CreateLambda([this]() {
-			check(OSCBufferLock == OSC_BUFFER_LOCKED);
-			FPlatformAtomics::InterlockedCompareExchange(&OSCBufferLock, OSC_BUFFER_FREE, OSC_BUFFER_LOCKED);
-				
-				TSharedPtr<FOSCPacket> packet;
-				while (OSCPackets.Buffer.Dequeue(packet))
-				{
-					if (packet->IsMessage())
-					{
-						FOSCMessage message;
-						message.SetPacket(StaticCastSharedPtr<FOSCMessagePacket>(packet));
-						OnOscReceived.Broadcast(message.GetPacket()->GetAddress(), message);
-					}
-					else if (packet->IsBundle())
-					{
-						FOSCBundle bundle;
-						bundle.SetPacket(StaticCastSharedPtr<FOSCBundlePacket>(packet));
-						OnOscBundleReceived.Broadcast(bundle);
-					}
-				}
-				OSCPackets.Buffer.Empty();
-			}),
-			TStatId(),
-			nullptr,
-			ENamedThreads::GameThread);
+		UE_LOG(LogOSC, Error, TEXT("Cannot update MulticastLoopback while OSCServer '%s' is active."), *GetName());
+		return;
 	}
+
+	bMulticastLoopback = bInMulticastLoopback;
 }
 
 void UOSCServer::Stop()
 {
-	delete SocketReceiver;
-	SocketReceiver = nullptr;
+	if (SocketReceiver)
+	{
+		delete SocketReceiver;
+		SocketReceiver = nullptr;
+	}
 
 	if (Socket)
 	{
+		Socket->Close();
 		ISocketSubsystem::Get(PLATFORM_SOCKETSUBSYSTEM)->DestroySocket(Socket);
 		Socket = nullptr;
 	}
@@ -128,3 +177,43 @@ void UOSCServer::BeginDestroy()
 	Stop();
 	Super::BeginDestroy();
 }
+
+void UOSCServer::AddWhitelistedClient(const FString& IPAddress)
+{
+	FIPv4Address OutAddress;
+	if (!FIPv4Address::Parse(IPAddress, OutAddress))
+	{
+		UE_LOG(LogOSC, Warning, TEXT("OSCServer '%s' failed to whitelist IP Address '%s'. Address is invalid."), *IPAddress);
+		return;
+	}
+
+	ClientWhitelist.Add(OutAddress.Value);
+}
+
+void UOSCServer::RemoveWhitelistedClient(const FString& IPAddress)
+{
+	FIPv4Address OutAddress;
+	if (!FIPv4Address::Parse(IPAddress, OutAddress))
+	{
+		UE_LOG(LogOSC, Warning, TEXT("OSCServer '%s' failed to remove whitelisted IP Address '%s'. Address is invalid."), *IPAddress);
+		return;
+	}
+
+	ClientWhitelist.Remove(OutAddress.Value);
+}
+
+void UOSCServer::ClearWhitelistedClients()
+{
+	ClientWhitelist.Reset();
+}
+
+TSet<FString> UOSCServer::GetWhitelistedClients() const
+{
+	TSet<FString> OutWhitelist;
+	for (uint32 Client : ClientWhitelist)
+	{
+		OutWhitelist.Add(FIPv4Address(Client).ToString());
+	}
+
+	return OutWhitelist;
+}
diff --git a/Engine/Plugins/Runtime/OSC/Source/OSC/Private/OSCStream.cpp b/Engine/Plugins/Runtime/OSC/Source/OSC/Private/OSCStream.cpp
--- a/Engine/Plugins/Runtime/OSC/Source/OSC/Private/OSCStream.cpp
+++ b/Engine/Plugins/Runtime/OSC/Source/OSC/Private/OSCStream.cpp
@@ -1,50 +1,88 @@
 // Copyright 1998-2019 Epic Games, Inc. All Rights Reserved.
 #include "OSCStream.h"
 
-FOSCStream::FOSCStream() : Data(nullptr), Size(0) { Position = 0; }
-FOSCStream::~FOSCStream() {}
-FOSCStream::FOSCStream(char* Data, int Size) 
-{ 
-	this->Data = Data;
-	this->Size = Size;
-	Position = 0; 
+#include "OSCLog.h"
+
+
+FOSCStream::FOSCStream()
+	: Position(0)
+{
+}
+
+FOSCStream::FOSCStream(const uint8* InData, int32 InSize) 
+	: Data(InData, InSize)
+	, Position(0)
+{
+}
+
+FOSCStream::FOSCStream(int32 InSize)
+	: Position(0)
+{
+	Data.AddZeroed(InSize);
+}
+
+const uint8* FOSCStream::GetData() const
+{
+	return Data.GetData();
+}
+
+int32 FOSCStream::GetLength() const
+{
+	return Data.Num();
+}
+
+bool FOSCStream::HasReachedEnd() const
+{
+	return Position >= Data.Num();
+}
+
+int32 FOSCStream::GetPosition() const
+{
+	return Position;
+}
+
+void FOSCStream::SetPosition(int32 InPosition)
+{
+	Position = InPosition;
 }
 
-char FOSCStream::ReadChar()
+TCHAR FOSCStream::ReadChar()
 {
-	char temp;
-	if (Read(&temp, 1) == 1)
+	uint8 Temp;
+	if (Read(&Temp, 1) > 0)
 	{
-		return temp;
+		TCHAR* OutChar = ANSI_TO_TCHAR((ANSICHAR*)&Temp);
+		return *OutChar;
 	}
 
-	return 0;
+	return '\0';
 }
 
-void FOSCStream::WriteChar(char Char)
+void FOSCStream::WriteChar(TCHAR Char)
 {
-	Write(&Char, 1);
+	const uint8* ToWrite = (uint8*)TCHAR_TO_ANSI(&Char);
+	Write(ToWrite, 1);
 }
 
 int32 FOSCStream::ReadInt32()
 {
-	char temp[4];
-	if (Read(temp, 4) == 4)
+	uint8 Temp[4];
+	if (Read(Temp, 4) == 4)
 	{
 #if PLATFORM_LITTLE_ENDIAN
 		union {
 			int32 i;
-			char c[4];
+			uint8 c[4];
 		} u;
 
-		u.c[0] = temp[3];
-		u.c[1] = temp[2];
-		u.c[2] = temp[1];
-		u.c[3] = temp[0];
+		u.c[0] = Temp[3];
+		u.c[1] = Temp[2];
+		u.c[2] = Temp[1];
+		u.c[3] = Temp[0];
 
 		return u.i;
 #else
-		return *(int32*)temp;
+		return *(int32*)(Temp);
 #endif
 	}
 
@@ -53,49 +91,49 @@ int32 FOSCStream::ReadInt32()
 
 void FOSCStream::WriteInt32(int32 Value)
 {
-	char temp[4];
+	uint8 Temp[4];
 
 #ifdef PLATFORM_LITTLE_ENDIAN
 	union {
 		int32 i;
-		char c[4];
+		uint8 c[4];
 	} u;
 
 	u.i = Value;
 
-	temp[3] = u.c[0];
-	temp[2] = u.c[1];
-	temp[1] = u.c[2];
-	temp[0] = u.c[3];
+	Temp[3] = u.c[0];
+	Temp[2] = u.c[1];
+	Temp[1] = u.c[2];
+	Temp[0] = u.c[3];
 #else
-	*reinterpret_cast<int32*>(temp) = Value;
+	*reinterpret_cast<int32*>(Temp) = Value;
 #endif
-	Write(temp, 4);
+	Write(Temp, 4);
 }
 
 double FOSCStream::ReadDouble()
 {
-	char temp[8];
-	if (Read(temp, 8) == 8)
+	uint8 Temp[8];
+	if (Read(Temp, 8) == 8)
 	{
 #if PLATFORM_LITTLE_ENDIAN
 		union {
 			double d;
-			char c[8];
+			uint8 c[8];
 		} u;
 
-		u.c[0] = temp[7];
-		u.c[1] = temp[6];
-		u.c[2] = temp[5];
-		u.c[3] = temp[4];
-		u.c[4] = temp[3];
-		u.c[5] = temp[2];
-		u.c[6] = temp[1];
-		u.c[7] = temp[0];
+		u.c[0] = Temp[7];
+		u.c[1] = Temp[6];
+		u.c[2] = Temp[5];
+		u.c[3] = Temp[4];
+		u.c[4] = Temp[3];
+		u.c[5] = Temp[2];
+		u.c[6] = Temp[1];
+		u.c[7] = Temp[0];
 
 		return u.d;
 #else
-		return *(double*)temp;
+		return *(double*)Temp;
 #endif
 
 	}
@@ -105,54 +143,53 @@ double FOSCStream::ReadDouble()
 
 void FOSCStream::WriteDouble(uint64 Value)
 {
-	char temp[8];
+	uint8 Temp[8];
 
 #ifdef PLATFORM_LITTLE_ENDIAN
 	union {
 		double i;
-		char c[8];
+		uint8 c[8];
 	} u;
 
 	u.i = Value;
 
-	temp[7] = u.c[0];
-	temp[6] = u.c[1];
-	temp[5] = u.c[2];
-	temp[4] = u.c[3];
-	temp[3] = u.c[4];
-	temp[2] = u.c[5];
-	temp[1] = u.c[6];
-	temp[0] = u.c[7];
+	Temp[7] = u.c[0];
+	Temp[6] = u.c[1];
+	Temp[5] = u.c[2];
+	Temp[4] = u.c[3];
+	Temp[3] = u.c[4];
+	Temp[2] = u.c[5];
+	Temp[1] = u.c[6];
+	Temp[0] = u.c[7];
 #else
-	*reinterpret_cast<double*>(temp) = Value;
+	*reinterpret_cast<double*>(Temp) = Value;
 #endif
-	Write(temp, 8);
+	Write(Temp, 8);
 }
 
-
 int64 FOSCStream::ReadInt64()
 {
-	char temp[8];
-	if (Read(temp, 8) == 8)
+	uint8 Temp[8];
+	if (Read(Temp, 8) == 8)
 	{
 #if PLATFORM_LITTLE_ENDIAN
 		union {
 			int64 i;
-			char c[8];
+			uint8 c[8];
 		} u;
 
-		u.c[0] = temp[7];
-		u.c[1] = temp[6];
-		u.c[2] = temp[5];
-		u.c[3] = temp[4];
-		u.c[4] = temp[3];
-		u.c[5] = temp[2];
-		u.c[6] = temp[1];
-		u.c[7] = temp[0];
+		u.c[0] = Temp[7];
+		u.c[1] = Temp[6];
+		u.c[2] = Temp[5];
+		u.c[3] = Temp[4];
+		u.c[4] = Temp[3];
+		u.c[5] = Temp[2];
+		u.c[6] = Temp[1];
+		u.c[7] = Temp[0];
 
 		return u.i;
 #else
-		return *(int64*)temp;
+		return *(int64*)Temp;
 #endif
 	}
 
@@ -161,53 +198,53 @@ int64 FOSCStream::ReadInt64()
 
 void FOSCStream::WriteInt64(int64 Value)
 {
-	char temp[8];
+	uint8 Temp[8];
 
 #ifdef PLATFORM_LITTLE_ENDIAN
 	union {
 		int64 i;
-		char c[8];
+		uint8 c[8];
 	} u;
 
 	u.i = Value;
 
-	temp[7] = u.c[0];
-	temp[6] = u.c[1];
-	temp[5] = u.c[2];
-	temp[4] = u.c[3];
-	temp[3] = u.c[4];
-	temp[2] = u.c[5];
-	temp[1] = u.c[6];
-	temp[0] = u.c[7];
+	Temp[7] = u.c[0];
+	Temp[6] = u.c[1];
+	Temp[5] = u.c[2];
+	Temp[4] = u.c[3];
+	Temp[3] = u.c[4];
+	Temp[2] = u.c[5];
+	Temp[1] = u.c[6];
+	Temp[0] = u.c[7];
 #else
-	*reinterpret_cast<int64*>(temp) = Value;
+	*reinterpret_cast<int64*>(Temp) = Value;
 #endif
-	Write(temp, 8);
+	Write(Temp, 8);
 }
 
 uint64 FOSCStream::ReadUInt64()
 {
-	char temp[8];
-	if (Read(temp, 8) == 8)
+	uint8 Temp[8];
+	if (Read(Temp, 8) == 8)
 	{
 #if PLATFORM_LITTLE_ENDIAN
 		union {
 			uint64 i;
-			char c[8];
+			uint8 c[8];
 		} u;
 
-		u.c[0] = temp[7];
-		u.c[1] = temp[6];
-		u.c[2] = temp[5];
-		u.c[3] = temp[4];
-		u.c[4] = temp[3];
-		u.c[5] = temp[2];
-		u.c[6] = temp[1];
-		u.c[7] = temp[0];
+		u.c[0] = Temp[7];
+		u.c[1] = Temp[6];
+		u.c[2] = Temp[5];
+		u.c[3] = Temp[4];
+		u.c[4] = Temp[3];
+		u.c[5] = Temp[2];
+		u.c[6] = Temp[1];
+		u.c[7] = Temp[0];
 
 		return u.i;
 #else
-		return *(uint64*)temp;
+		return *(uint64*)Temp;
 #endif
 
 	}
@@ -217,49 +254,49 @@ uint64 FOSCStream::ReadUInt64()
 
 void FOSCStream::WriteUInt64(uint64 Value)
 {
-	char temp[8];
+	uint8 Temp[8];
 
 #ifdef PLATFORM_LITTLE_ENDIAN
 	union {
 		uint64 i;
-		char c[8];
+		uint8 c[8];
 	} u;
 
 	u.i = Value;
 
-	temp[7] = u.c[0];
-	temp[6] = u.c[1];
-	temp[5] = u.c[2];
-	temp[4] = u.c[3];
-	temp[3] = u.c[4];
-	temp[2] = u.c[5];
-	temp[1] = u.c[6];
-	temp[0] = u.c[7];
+	Temp[7] = u.c[0];
+	Temp[6] = u.c[1];
+	Temp[5] = u.c[2];
+	Temp[4] = u.c[3];
+	Temp[3] = u.c[4];
+	Temp[2] = u.c[5];
+	Temp[1] = u.c[6];
+	Temp[0] = u.c[7];
 #else
-	*reinterpret_cast<uint64*>(temp) = Value;
+	*reinterpret_cast<uint64*>(Temp) = Value;
 #endif
-	Write(temp, 8);
+	Write(Temp, 8);
 }
 
 float FOSCStream::ReadFloat()
 {
-	char temp[4];
-	if (Read(temp, 4) == 4)
+	uint8 Temp[4];
+	if (Read(Temp, 4) == 4)
 	{
 #if PLATFORM_LITTLE_ENDIAN
 		union {
 			float f;
-			char c[4];
+			uint8 c[4];
 		} u;
 
-		u.c[0] = temp[3];
-		u.c[1] = temp[2];
-		u.c[2] = temp[1];
-		u.c[3] = temp[0];
+		u.c[0] = Temp[3];
+		u.c[1] = Temp[2];
+		u.c[2] = Temp[1];
+		u.c[3] = Temp[0];
 
 		return u.f;
 #else
-		return *(float*)temp;
+		return *(float*)Temp;
 #endif
 	}
 
@@ -268,56 +305,72 @@ float FOSCStream::ReadFloat()
 
 void FOSCStream::WriteFloat(float Value)
 {
-	char temp[4];
+	uint8 Temp[4];
 
 #ifdef PLATFORM_LITTLE_ENDIAN
 	union {
 		float f;
-		char c[4];
+		uint8 c[4];
 	} u;
 
 	u.f = Value;
 
-	temp[3] = u.c[0];
-	temp[2] = u.c[1];
-	temp[1] = u.c[2];
-	temp[0] = u.c[3];
-#else
-	*reinterpret_cast<float*>(temp) = Value;
-#endif
-	Write(temp, 4);
+	Temp[3] = u.c[0];
+	Temp[2] = u.c[1];
+	Temp[1] = u.c[2];
+	Temp[0] = u.c[3];
+#else // if !PLATFORM_LITTLE_ENDIAN
+	*reinterpret_cast<float*>(Temp) = Value;
+#endif // !PLATFORM_LITTLE_ENDIAN
+
+	Write(Temp, 4);
 }
 
 FString FOSCStream::ReadString()
 {
-	int32 count = 0;
-	for (int32 index = Position; Data[index] != 0; index++)
+	const int32 DataSize = Data.Num();
+	
+	const int32 InitPosition = Position;
+	const ANSICHAR* StrStart = (ANSICHAR*)(&Data[InitPosition]);
+
+	int32 i = InitPosition;
+	for (; i < DataSize; i++)
 	{
-		count++;
+		Position++;
+		if (Data[i] == '\0')
+		{
+			break;
+		}
 	}
 
-	FString s(count, &Data[Position]);
+	if (Position == DataSize)
+	{
+		UE_LOG(LogOSC, Error, TEXT("Invalid string: Terminator not found"));
+		return FString();
+	}
 
-	Position += count + 1; // increment Position
-	Position = ((Position + 3) / 4) * 4; // padded
+	// Pad position before returning
+	const int32 EndPosition = Position;
+	Position = ((Position + 3) / 4) * 4;
 
-	return s;
+	return FString(EndPosition - InitPosition, StrStart);
 }
 
-void FOSCStream::WriteString(FString string)
+void FOSCStream::WriteString(const FString& InString)
 {
-	TArray<TCHAR> s = string.GetCharArray();
+	const TArray<TCHAR>& CharArr = InString.GetCharArray();
 
-	const int32 count = s.Num();
-	for (int32 index = 0; index < count; index++)
+	const int32 Count = CharArr.Num();
+	for (int32 i = 0; i < Count; i++)
 	{
-		Data[Position + index] = s[index];
+		Data[Position + i] = CharArr[i];
 	}
 
-	Position += count;
+	Position += Count;
 
-	const int32 numPaddingZeros = ((count + 3) / 4) * 4; // padded
-	for (int32 i = 0; i < numPaddingZeros - count; i++)
+	// Increment & pad string with null terminator
+	const int32 NumPaddingZeros = ((Count + 3) / 4) * 4;
+	for (int32 i = 0; i < NumPaddingZeros - Count; i++)
 	{
 		WriteChar('\0');
 	}
@@ -327,9 +380,8 @@ TArray<uint8> FOSCStream::ReadBlob()
 {
 	TArray<uint8> Blob;
 
-	int blobsize = ReadInt32();
-
-	for (int i = 0; i < blobsize; i++)
+	const int32 BlobSize = ReadInt32();
+	for (int32 i = 0; i < BlobSize; i++)
 	{
 		Blob.Add(ReadChar());
 	}
@@ -341,40 +393,40 @@ TArray<uint8> FOSCStream::ReadBlob()
 
 void FOSCStream::WriteBlob(TArray<uint8>& Blob)
 {
-	// Write Blob size
 	WriteInt32(Blob.Num());
-
-	for (int i = 0; i < Blob.Num(); i++)
+	for (int32 i = 0; i < Blob.Num(); i++)
 	{
-		Write(&Blob[i], 4);
+		Write(&Blob[i], 4u);
 	}
 }
 
-int FOSCStream::Read(void* Buffer, int ToRead)
+int32 FOSCStream::Read(uint8* Buffer, int32 InSize)
 {
-	if (ToRead <= 0 || Position >= Size)
+	const int32 DataSize = Data.Num();
+	if (InSize == 0 || Position >= DataSize)
+	{
 		return 0;
+	}
 
-	int num = FMath::Min<int>(ToRead, Size - Position);
-	if (num > 0)
+	const int32 Num = FMath::Min<int32>(InSize, DataSize - Position);
+	if (Num > 0)
 	{
-		memcpy(Buffer, Data + Position, num);
-
-		Position += num;
+		FMemory::Memcpy(Buffer, &Data[Position], Num);
+		Position += Num;
 	}
 
-	return num;
+	return Num;
 }
 
-int FOSCStream::Write(void* Buffer, int ToWrite)
+int32 FOSCStream::Write(const uint8* InBuffer, int32 InSize)
 {
-	if (ToWrite <= 0)
-		return 0;
-
-	memcpy(Data + Position, Buffer, ToWrite);
-
-	Position += ToWrite;
+	if (InSize > 0)
+	{
+		FMemory::Memcpy(&Data[Position], InBuffer, InSize);
+		Position += InSize;
+		return InSize;
+	}
 
-	return ToWrite;
+	return 0;
 }
 
diff --git a/Engine/Plugins/Runtime/OSC/Source/OSC/Public/OSCAddress.h b/Engine/Plugins/Runtime/OSC/Source/OSC/Public/OSCAddress.h
new file mode 100644
--- /dev/null
+++ b/Engine/Plugins/Runtime/OSC/Source/OSC/Public/OSCAddress.h
@@ -0,0 +1,80 @@
+// Copyright 1998-2019 Epic Games, Inc. All Rights Reserved.
+#pragma once
+
+#include "CoreMinimal.h"
+#include "UObject/ObjectMacros.h"
+#include "UObject/Object.h"
+#include "OSCLog.h"
+
+#include "OSCAddress.generated.h"
+
+
+namespace OSC
+{
+	const FString BundleTag		= TEXT("#bundle");
+	const FString PathSeparator = TEXT("/");
+} // namespace OSC
+
+USTRUCT(BlueprintType)
+struct OSC_API FOSCAddress
+{
+	GENERATED_USTRUCT_BODY()
+
+public:
+	FOSCAddress()
+	{
+		Value = OSC::PathSeparator;
+	}
+
+	FOSCAddress(const FString& InValue)
+		: Value(InValue)
+	{
+		if (!IsBundle())
+		{
+			if (!Value.StartsWith(OSC::PathSeparator))
+			{
+				Value = OSC::PathSeparator + Value;
+			}
+		}
+	}
+
+	UPROPERTY(BlueprintReadWrite, Category = "Audio|OSC")
+	FString Value;
+
+	FORCEINLINE bool IsBundle() const
+	{
+		return Value.StartsWith(OSC::BundleTag);
+	}
+
+	FORCEINLINE bool IsMessage() const
+	{
+		return Value.StartsWith(OSC::PathSeparator);
+	}
+
+	FORCEINLINE bool IsValid() const
+	{
+		return IsMessage() || IsBundle();
+	}
+
+	FORCEINLINE void Append(const FString& ToAppend)
+	{
+		if (!IsMessage())
+		{
+			UE_LOG(LogOSC, Warning, TEXT("Append failed for OSCAddress '%s'. Address is either invalid/non-message type."), *Value);
+			return;
+		}
+
+		Value.PathAppend(*ToAppend, ToAppend.Len());
+	}
+
+	FORCEINLINE void Split(TArray<FString>& OutArray) const
+	{
+		if (!IsMessage())
+		{
+			UE_LOG(LogOSC, Warning, TEXT("Split failed for OSCAddress '%s'. Address is invalid/non-message type."), *Value);
+			return;
+		}
+
+		Value.ParseIntoArray(OutArray, *OSC::PathSeparator, true);
+	}
+};
diff --git a/Engine/Plugins/Runtime/OSC/Source/OSC/Public/OSCBundle.h b/Engine/Plugins/Runtime/OSC/Source/OSC/Public/OSCBundle.h
new file mode 100644
--- /dev/null
+++ b/Engine/Plugins/Runtime/OSC/Source/OSC/Public/OSCBundle.h
@@ -0,0 +1,40 @@
+// Copyright 1998-2019 Epic Games, Inc. All Rights Reserved.
+#pragma once 
+
+#include "OSCPacket.h"
+#include "OSCBundle.generated.h"
+
+
+USTRUCT(BlueprintType)
+struct OSC_API FOSCBundle
+{
+	GENERATED_USTRUCT_BODY()
+
+	FOSCBundle()
+		: Packet(MakeShareable(new FOSCBundlePacket()))
+	{
+	}
+
+	FOSCBundle(const TSharedPtr<FOSCBundlePacket>& InPacket)
+		: Packet(InPacket)
+	{
+	}
+
+	~FOSCBundle()
+	{
+		Packet.Reset();
+	}
+
+	void SetPacket(const TSharedPtr<FOSCBundlePacket>& InPacket)
+	{
+		Packet = InPacket;
+	}
+
+	const TSharedPtr<FOSCBundlePacket>& GetPacket() const
+	{
+		return Packet;
+	}
+
+private:
+	TSharedPtr<FOSCBundlePacket> Packet;
+};
\ No newline at end of file
diff --git a/Engine/Plugins/Runtime/OSC/Source/OSC/Public/OSCClient.h b/Engine/Plugins/Runtime/OSC/Source/OSC/Public/OSCClient.h
new file mode 100644
--- /dev/null
+++ b/Engine/Plugins/Runtime/OSC/Source/OSC/Public/OSCClient.h
@@ -0,0 +1,55 @@
+// Copyright 1998-2019 Epic Games, Inc. All Rights Reserved.
+#pragma once
+
+#include "CoreMinimal.h"
+#include "UObject/Object.h"
+#include "Interfaces/IPv4/IPv4Address.h"
+#include "IPAddress.h"
+#include "OSCMessage.h"
+#include "OSCBundle.h"
+#include "OSCClient.generated.h"
+
+class FSocket;
+
+UCLASS(BlueprintType)
+class OSC_API UOSCClient : public UObject
+{
+	GENERATED_BODY()
+
+public:
+	UOSCClient();
+	virtual ~UOSCClient();
+
+	/** Sets the OSC Client IP address and port. Returns whether 
+	  * address and port was successfully set. */
+	UFUNCTION(BlueprintCallable, Category = "Audio|OSC")
+	bool SetSendIPAddress(const FString& IPAddress, const int32 Port);
+
+	/** Gets the OSC Client IP address and port. */
+	UFUNCTION(BlueprintCallable, Category = "Audio|OSC")
+	void GetSendIPAddress(UPARAM(ref) FString& IPAddress, UPARAM(ref) int32& Port);
+
+	/** Send OSC message to  a specific address. */
+	UFUNCTION(BlueprintCallable, Category = "Audio|OSC")
+	void SendOSCMessage(UPARAM(ref) FOSCMessage& Message);
+
+	/** Send OSC Bundle over the network. */
+	UFUNCTION(BlueprintCallable, Category = "Audio|OSC")
+	void SendOSCBundle(UPARAM(ref) FOSCBundle& Bundle);
+
+protected:
+
+	void BeginDestroy() override;
+	
+	/** Stop and tidy up network socket. */
+	void Stop();
+	
+	/** Send OSC packet data. */
+	void SendPacket(FOSCPacket& InPacket);
+
+	/** Socket used to send the OSC packets. */
+	FSocket* Socket;
+
+	/** IP Address used by socket. */
+	TSharedPtr<FInternetAddr> IPAddress;
+};
diff --git a/Engine/Plugins/Runtime/OSC/Source/OSC/Public/OSCLog.h b/Engine/Plugins/Runtime/OSC/Source/OSC/Public/OSCLog.h
new file mode 100644
--- /dev/null
+++ b/Engine/Plugins/Runtime/OSC/Source/OSC/Public/OSCLog.h
@@ -0,0 +1,6 @@
+// Copyright 1998-2019 Epic Games, Inc. All Rights Reserved.
+#pragma once
+
+#include "CoreMinimal.h"
+
+DECLARE_LOG_CATEGORY_EXTERN(LogOSC, Log, All);
diff --git a/Engine/Plugins/Runtime/OSC/Source/OSC/Public/OSCManager.cpp b/Engine/Plugins/Runtime/OSC/Source/OSC/Public/OSCManager.cpp
--- a/Engine/Plugins/Runtime/OSC/Source/OSC/Public/OSCManager.cpp
+++ b/Engine/Plugins/Runtime/OSC/Source/OSC/Public/OSCManager.cpp
@@ -1,284 +1,374 @@
 // Copyright 1998-2019 Epic Games, Inc. All Rights Reserved.
 #include "OSCManager.h"
+
+#include "IPAddress.h"
 #include "OSCMessage.h"
 #include "OSCBundle.h"
 #include "OSCLog.h"
+#include "SocketSubsystem.h"
+
+
+namespace
+{
+	// Returns true if provided address was null and was able to
+	// override with local host address, false if not.
+	bool GetLocalHostAddress(FString& Address)
+	{
+		if (!Address.IsEmpty() && Address != TEXT("0"))
+		{
+			return false;
+		}
+
+		bool bCanBind = false;
+		bool bAppendPort = false;
+		if (ISocketSubsystem* SocketSys = ISocketSubsystem::Get(PLATFORM_SOCKETSUBSYSTEM))
+		{
+			const TSharedPtr<FInternetAddr> Addr = SocketSys->GetLocalHostAddr(*GLog, bCanBind);
+			if (Addr.IsValid())
+			{
+				Address = Addr->ToString(bAppendPort);
+				return true;
+			}
+		}
+
+		return false;
+	}
+
+	const FOSCType* GetOSCTypeAtIndex(const FOSCMessage& Message, const int32 Index)
+	{
+		const TSharedPtr<FOSCMessagePacket>& Packet = Message.GetPacket();
+		if (Packet.IsValid())
+		{
+			TArray<FOSCType>& Args = Packet->GetArguments();
+			if (Index >= Args.Num())
+			{
+				UE_LOG(LogOSC, Error, TEXT("Index '%d' out-of-bounds.  Message argument size = '%d'"), Index, Args.Num());
+				return nullptr;
+			}
 
-UOSCServer* UOSCManager::CreateOSCServer(FString ReceiveAddress, int32 Port)
+			return &Args[Index];
+		}
+
+		return nullptr;
+	}
+} // namespace <>
+
+UOSCServer* UOSCManager::CreateOSCServer(FString ReceiveIPAddress, int32 Port, bool bMulticastLoopback, bool bStartListening)
 {
 	UOSCServer* NewOSCServer = NewObject<UOSCServer>();
 
-	FIPv4Address ipAddr;
-	if (FIPv4Address::Parse(ReceiveAddress, ipAddr))
+	if (GetLocalHostAddress(ReceiveIPAddress))
 	{
-		NewOSCServer->Listen(ipAddr, Port, false);
+		UE_LOG(LogOSC, Display, TEXT("OSCServer ReceiveAddress not specified. Using LocalHost IP: '%s'"), *ReceiveIPAddress);
+	}
+
+	NewOSCServer->SetMulticastLoopback(bMulticastLoopback);
+	if (NewOSCServer->SetAddress(ReceiveIPAddress, Port))
+	{
+		if (bStartListening)
+		{
+			NewOSCServer->Listen();
+		}
+	}
+	else
+	{
+		UE_LOG(LogOSC, Warning, TEXT("Failed to parse ReceiveAddress '%s' for OSCServer."), *ReceiveIPAddress);
 	}
 
 	return NewOSCServer;
 }
 
-UOSCClient* UOSCManager::CreateOSCClient(FString SendAddress, int32 Port)
+UOSCClient* UOSCManager::CreateOSCClient(FString SendIPAddress, int32 Port)
 {
 	UOSCClient* NewOSCClient = NewObject<UOSCClient>();
 
-	FIPv4Address ipAddr;
-	if (FIPv4Address::Parse(SendAddress, ipAddr))
+	if (GetLocalHostAddress(SendIPAddress))
 	{
-		NewOSCClient->SetTarget(ipAddr, Port);
+		UE_LOG(LogOSC, Display, TEXT("OSCClient SendAddress not specified. Using LocalHost IP: '%s'"), *SendIPAddress);
 	}
 
-	return NewOSCClient;
-}
-
-void UOSCManager::CreateOSCMessage(UPARAM(ref) FOSCMessage& Message)
-{
-	Message.SetPacket(MakeShareable(new FOSCMessagePacket()));
-}
+	if (!NewOSCClient->SetSendIPAddress(SendIPAddress, Port))
+	{
+		UE_LOG(LogOSC, Warning, TEXT("Failed to parse SendAddress '%s' for OSCClient. Client unable to send new messages."), *SendIPAddress);
+	}
 
-void UOSCManager::CreateOSCBundle(UPARAM(ref) FOSCBundle& Bundle)
-{
-	Bundle.SetPacket(MakeShareable(new FOSCBundlePacket()));
+	return NewOSCClient;
 }
 
-void UOSCManager::ClearMessage(UPARAM(ref) FOSCMessage& Message)
+void UOSCManager::ClearMessage(FOSCMessage& Message)
 {
-	if (Message.IsValid())
+	if (Message.GetPacket().IsValid())
 	{
-		Message.GetPacket()->Clear();
+		Message.GetPacket()->GetArguments().Reset();
 	}
 }
 
-void UOSCManager::ClearBundle(UPARAM(ref) FOSCBundle& Bundle)
+void UOSCManager::ClearBundle(FOSCBundle& Bundle)
 {
-	if (Bundle.IsValid())
+	if (Bundle.GetPacket().IsValid())
 	{
-		Bundle.GetPacket()->Clear();
+		Bundle.GetPacket()->GetPackets().Reset();
 	}
 }
 
-void UOSCManager::AddMessageToBundle(const FOSCMessage& Message, UPARAM(ref) FOSCBundle& Bundle)
+FOSCBundle& UOSCManager::AddMessageToBundle(const FOSCMessage& Message, FOSCBundle& Bundle)
 {
-	if (Message.IsValid() && Bundle.IsValid())
+	if (Message.GetPacket().IsValid() && Bundle.GetPacket().IsValid())
 	{
-		Bundle.GetPacket()->AddPacket(Message.GetPacket());
+		Bundle.GetPacket()->GetPackets().Add(Message.GetPacket());
 	}
+
+	return Bundle;
 }
 
-void UOSCManager::AddBundleToBundle(const FOSCBundle& inBundle, UPARAM(ref) FOSCBundle& outBundle)
+FOSCBundle& UOSCManager::AddBundleToBundle(const FOSCBundle& InBundle, FOSCBundle& OutBundle)
 {
-	if (inBundle.IsValid() && outBundle.IsValid())
+	if (InBundle.GetPacket().IsValid() && OutBundle.GetPacket().IsValid())
 	{
-		outBundle.GetPacket()->AddPacket(inBundle.GetPacket());
+		OutBundle.GetPacket()->GetPackets().Add(InBundle.GetPacket());
 	}
-}
 
-void UOSCManager::SetAddress(UPARAM(ref) FOSCMessage& Message, const FString& Address)
-{
-	Message.GetOrCreatePacket()->SetAddress(Address);
+	return OutBundle;
 }
 
-void UOSCManager::AddFloat(UPARAM(ref) FOSCMessage& Message, float Value)
+FOSCMessage& UOSCManager::AddFloat(FOSCMessage& Message, float Value)
 {
-	Message.GetOrCreatePacket()->AddArgument(FOSCType(Value));
+	Message.GetPacket()->GetArguments().Add(FOSCType(Value));
+	return Message;
 }
 
-void UOSCManager::AddInt32(UPARAM(ref) FOSCMessage& Message, int32 Value)
+FOSCMessage& UOSCManager::AddInt32(FOSCMessage& Message, int32 Value)
 {
-	Message.GetOrCreatePacket()->AddArgument(FOSCType(Value));
+	Message.GetPacket()->GetArguments().Add(FOSCType(Value));
+	return Message;
 }
 
-void UOSCManager::AddInt64(UPARAM(ref) FOSCMessage& Message, int64 Value)
+FOSCMessage& UOSCManager::AddInt64(FOSCMessage& Message, int64 Value)
 {
-	Message.GetOrCreatePacket()->AddArgument(FOSCType(Value));
+	Message.GetPacket()->GetArguments().Add(FOSCType(Value));
+	return Message;
 }
 
-void UOSCManager::AddString(UPARAM(ref) FOSCMessage& Message, FString Value)
+FOSCMessage& UOSCManager::AddString(FOSCMessage& Message, FString Value)
 {
-	Message.GetOrCreatePacket()->AddArgument(FOSCType(Value));
+	Message.GetPacket()->GetArguments().Add(FOSCType(Value));
+	return Message;
 }
 
-void UOSCManager::AddBlob(UPARAM(ref) FOSCMessage& Message, TArray<uint8>& Value)
+FOSCMessage& UOSCManager::AddBlob(FOSCMessage& Message, TArray<uint8>& Value)
 {
-	Message.GetOrCreatePacket()->AddArgument(FOSCType(Value));
+	Message.GetPacket()->GetArguments().Add(FOSCType(Value));
+	return Message;
 }
 
-void UOSCManager::AddBool(UPARAM(ref) FOSCMessage& Message, bool Value)
+FOSCMessage& UOSCManager::AddBool(FOSCMessage& Message, bool Value)
 {
-	Message.GetOrCreatePacket()->AddArgument(FOSCType(Value));
+	Message.GetPacket()->GetArguments().Add(FOSCType(Value));
+	return Message;
 }
 
 void UOSCManager::GetMessagesFromBundle(const FOSCBundle& Bundle, TArray<FOSCMessage>& Messages)
 {
-	if (Bundle.IsValid())
+	if (Bundle.GetPacket().IsValid())
 	{
-		const TSharedPtr<FOSCBundlePacket> bundlePacket = Bundle.GetPacket();
-
-		const int32 num = bundlePacket->GetNumPackets();
-		for (int32 i = 0; i < num; i++)
+		const TSharedPtr<FOSCBundlePacket> BundlePacket = Bundle.GetPacket();
+		for (int32 i = 0; i < BundlePacket->GetPackets().Num(); i++)
 		{
-			TSharedPtr<FOSCPacket> packet = bundlePacket->GetPacket(i);
-
-			if (packet->IsMessage())
+			const TSharedPtr<FOSCPacket>& Packet = BundlePacket->GetPackets()[i];
+			if (Packet->IsMessage())
 			{
-				FOSCMessage message;
-
-				message.SetPacket(StaticCastSharedPtr<FOSCMessagePacket>(packet));
-
-				Messages.Add(message);
+				Messages.Emplace(StaticCastSharedPtr<FOSCMessagePacket>(Packet));
 			}
 		}
 	}
 }
 
-void UOSCManager::GetFloat(const FOSCMessage& Message, const int index, float& Value)
+void UOSCManager::GetFloat(const FOSCMessage& Message, const int32 Index, float& Value)
 {
-	if (Message.IsValid())
+	if (const FOSCType* OSCType = GetOSCTypeAtIndex(Message, Index))
 	{
-		const FOSCType& t = Message.GetPacket()->GetArgument(index);
-
-		if (t.IsFloat())
+		if (OSCType->IsFloat())
 		{
-			Value = t.GetFloat();
+			Value = OSCType->GetFloat();
 		}
 	}
 }
 
-void UOSCManager::GetAllFloat(const FOSCMessage& Message, TArray<float>& Values)
+void UOSCManager::GetAllFloats(const FOSCMessage& Message, TArray<float>& Values)
 {
-	if (Message.IsValid())
+	if (Message.GetPacket().IsValid())
 	{
-		for (int i = 0; i < Message.GetPacket()->GetNumArguments(); i++)
+		const TArray<FOSCType>& Args = Message.GetPacket()->GetArguments();
+		for (int32 i = 0; i < Args.Num(); i++)
 		{
-			const FOSCType& t = Message.GetPacket()->GetArgument(i);
-			if (t.IsFloat())
+			const FOSCType& OSCType = Args[i];
+			if (OSCType.IsFloat())
 			{
-				Values.Add(t.GetFloat());
+				Values.Add(OSCType.GetFloat());
 			}
 		}
 	}
 }
 
-void UOSCManager::GetInt32(const FOSCMessage& Message, const int index, int32& Value)
+void UOSCManager::GetInt32(const FOSCMessage& Message, const int32 Index, int32& Value)
 {
-	if (Message.IsValid())
+	if (const FOSCType* OSCType = GetOSCTypeAtIndex(Message, Index))
 	{
-		const FOSCType& t = Message.GetPacket()->GetArgument(index);
-
-		if (t.IsInt32())
+		if (OSCType->IsInt32())
 		{
-			Value = t.GetInt32();
+			Value = OSCType->GetInt32();
 		}
 	}
 }
 
-void UOSCManager::GetAllInt32(const FOSCMessage& Message, TArray<int32>& Values)
+void UOSCManager::GetAllInt32s(const FOSCMessage& Message, TArray<int32>& Values)
 {
-	if (Message.IsValid())
+	if (Message.GetPacket().IsValid())
 	{
-		for (int i = 0; i < Message.GetPacket()->GetNumArguments(); i++)
+		const TArray<FOSCType>& Args = Message.GetPacket()->GetArguments();
+		for (int32 i = 0; i < Args.Num(); i++)
 		{
-			const FOSCType& t = Message.GetPacket()->GetArgument(i);
-			if (t.IsInt32())
+			const FOSCType& OSCType = Args[i];
+			if (OSCType.IsInt32())
 			{
-				Values.Add(t.GetInt32());
+				Values.Add(OSCType.GetInt32());
 			}
 		}
 	}
 }
 
-void UOSCManager::GetInt64(const FOSCMessage& Message, const int index, int64& Value)
+void UOSCManager::GetInt64(const FOSCMessage& Message, const int32 Index, int64& Value)
 {
-	if (Message.IsValid())
+	if (const FOSCType* OSCType = GetOSCTypeAtIndex(Message, Index))
 	{
-		const FOSCType& t = Message.GetPacket()->GetArgument(index);
-
-		if (t.IsInt64())
+		if (OSCType->IsInt64())
 		{
-			Value = t.GetInt64();
+			Value = OSCType->GetInt64();
 		}
 	}
 }
 
-void UOSCManager::GetAllInt64(const FOSCMessage& Message, TArray<int64>& Values)
+void UOSCManager::GetAllInt64s(const FOSCMessage& Message, TArray<int64>& Values)
 {
-	if (Message.IsValid())
+	if (Message.GetPacket().IsValid())
 	{
-		for (int i = 0; i < Message.GetPacket()->GetNumArguments(); i++)
+		const TArray<FOSCType>& Args = Message.GetPacket()->GetArguments();
+		for (int32 i = 0; i < Args.Num(); i++)
 		{
-			const FOSCType& t = Message.GetPacket()->GetArgument(i);
-			if (t.IsInt64())
+			const FOSCType& OSCType = Args[i];
+			if (OSCType.IsInt64())
 			{
-				Values.Add(t.GetInt64());
+				Values.Add(OSCType.GetInt64());
 			}
 		}
 	}
 }
 
-void UOSCManager::GetString(const FOSCMessage& Message, const int index, FString& Value)
+void UOSCManager::GetString(const FOSCMessage& Message, const int32 Index, FString& Value)
 {
-	if (Message.IsValid())
+	if (const FOSCType* OSCType = GetOSCTypeAtIndex(Message, Index))
 	{
-		const FOSCType& t = Message.GetPacket()->GetArgument(index);
-
-		if (t.IsString())
+		if (OSCType->IsString())
 		{
-			Value = t.GetString();
+			Value = OSCType->GetString();
 		}
 	}
 }
 
-void UOSCManager::GetAllString(const FOSCMessage& Message, TArray<FString>& Values)
+void UOSCManager::GetAllStrings(const FOSCMessage& Message, TArray<FString>& Values)
 {
-	if (Message.IsValid())
+	if (Message.GetPacket().IsValid())
 	{
-		for (int i = 0; i < Message.GetPacket()->GetNumArguments(); i++)
+		const TArray<FOSCType>& Args = Message.GetPacket()->GetArguments();
+		for (int32 i = 0; i < Args.Num(); i++)
 		{
-			const FOSCType& t = Message.GetPacket()->GetArgument(i);
-			if (t.IsString())
+			const FOSCType& OSCType = Args[i];
+			if (OSCType.IsString())
 			{
-				Values.Add(t.GetString());
+				Values.Add(OSCType.GetString());
 			}
 		}
 	}
 }
 
-void UOSCManager::GetBool(const FOSCMessage& Message, const int index, bool& Value)
+void UOSCManager::GetBool(const FOSCMessage& Message, const int32 Index, bool& Value)
 {
-	if (Message.IsValid())
+	if (const FOSCType* OSCType = GetOSCTypeAtIndex(Message, Index))
 	{
-		const FOSCType& t = Message.GetPacket()->GetArgument(index);
-
-		if (t.IsBool())
+		if (OSCType->IsBool())
 		{
-			Value = t.GetBool();
+			Value = OSCType->GetBool();
 		}
 	}
 }
 
-void UOSCManager::GetAllBool(const FOSCMessage& Message, TArray<bool>& Values)
+void UOSCManager::GetAllBools(const FOSCMessage& Message, TArray<bool>& Values)
 {
-	if (Message.IsValid())
+	if (Message.GetPacket().IsValid())
 	{
-		for (int i = 0; i < Message.GetPacket()->GetNumArguments(); i++)
+		const TArray<FOSCType>& Args = Message.GetPacket()->GetArguments();
+		for (int32 i = 0; i < Args.Num(); i++)
 		{
-			const FOSCType& t = Message.GetPacket()->GetArgument(i);
-			if (t.IsBool())
+			const FOSCType& OSCType = Args[i];
+			if (OSCType.IsBool())
 			{
-				Values.Add(t.GetBool());
+				Values.Add(OSCType.GetBool());
 			}
 		}
 	}
 }
 
-void UOSCManager::GetBlob(const FOSCMessage& Message, const int index, TArray<uint8>& Value)
+void UOSCManager::GetBlob(const FOSCMessage& Message, const int32 Index, TArray<uint8>& Value)
 {
-	if (Message.IsValid())
+	if (const FOSCType* OSCType = GetOSCTypeAtIndex(Message, Index))
 	{
-		const FOSCType& t = Message.GetPacket()->GetArgument(index);
-
-		if (t.IsBlob())
+		if (OSCType->IsBlob())
 		{
-			Value = t.GetBlob();
+			Value = OSCType->GetBlob();
 		}
 	}
 }
+
+bool UOSCManager::OSCAddressIsBundle(const FOSCAddress& Address)
+{
+	return Address.IsBundle();
+}
+
+bool UOSCManager::OSCAddressIsMessage(const FOSCAddress& Address)
+{
+	return Address.IsBundle();
+}
+
+bool UOSCManager::OSCAddressIsValid(const FOSCAddress& Address)
+{
+	return Address.IsValid();
+}
+
+FOSCAddress UOSCManager::StringToOSCAddress(const FString& String)
+{
+	return FOSCAddress(String);
+}
+
+FOSCAddress& UOSCManager::OSCAddressAppend(FOSCAddress& Address, const FString& ToAppend)
+{
+	Address.Append(ToAppend);
+	return Address;
+}
+
+FOSCAddress UOSCManager::GetOSCMessageAddress(const FOSCMessage& Message)
+{
+	return Message.GetAddress();
+}
+
+FOSCMessage& UOSCManager::SetOSCMessageAddress(FOSCMessage& Message, const FOSCAddress& Address)
+{
+	Message.SetAddress(Address);
+	return Message;
+}
+
+TArray<FString> UOSCManager::SplitOSCAddress(const FOSCAddress& Address)
+{
+	TArray<FString> OutArray;
+	Address.Split(OutArray);
+	return MoveTemp(OutArray);
+}
\ No newline at end of file
diff --git a/Engine/Plugins/Runtime/OSC/Source/OSC/Public/OSCManager.h b/Engine/Plugins/Runtime/OSC/Source/OSC/Public/OSCManager.h
--- a/Engine/Plugins/Runtime/OSC/Source/OSC/Public/OSCManager.h
+++ b/Engine/Plugins/Runtime/OSC/Source/OSC/Public/OSCManager.h
@@ -14,88 +14,104 @@ class OSC_API UOSCManager : public UBlueprintFunctionLibrary
 	GENERATED_BODY()
 
 public:
-	//////////////////////////////////////////////////////////////////////////
-	// Audio|OSC functions
+	// Creates an OSC Server.  If left empty (or '0'),
+	// attempts to use LocalHost IP address. If StartListening set,
+	// immediately begins listening on creation.
 	UFUNCTION(BlueprintCallable, Category = "Audio|OSC")
-	static class UOSCServer* CreateOSCServer(FString ReceiveAddress, int32 Port);
+	static class UOSCServer* CreateOSCServer(FString ReceiveIPAddress, int32 Port, bool bMulticastLoopback, bool bStartListening);
 
+	// Creates an OSC Client.  If left empty (or '0'), attempts to use
+	// attempts to use LocalHost IP address.
 	UFUNCTION(BlueprintCallable, Category = "Audio|OSC")
-	static class UOSCClient* CreateOSCClient(FString SendAddress, int32 Port);
+	static class UOSCClient* CreateOSCClient(FString SendIPAddress, int32 Port);
 
 	UFUNCTION(BlueprintCallable, Category = "Audio|OSC")
-	static void CreateOSCMessage(UPARAM(ref) FOSCMessage& Message);
+	static UPARAM(DisplayName = "Bundle") FOSCBundle& AddMessageToBundle(const FOSCMessage& Message, UPARAM(ref) FOSCBundle& Bundle);
 
 	UFUNCTION(BlueprintCallable, Category = "Audio|OSC")
-	static void CreateOSCBundle(UPARAM(ref) FOSCBundle& Bundle);
+	static UPARAM(DisplayName = "OutBundle") FOSCBundle& AddBundleToBundle(const FOSCBundle& InBundle, UPARAM(ref) FOSCBundle& OutBundle);
 
 	UFUNCTION(BlueprintCallable, Category = "Audio|OSC")
-	static void AddMessageToBundle(const FOSCMessage& Message, UPARAM(ref) FOSCBundle& Bundle);
+	static void GetMessagesFromBundle(const FOSCBundle& Bundle, TArray<FOSCMessage>& Messages);
 
 	UFUNCTION(BlueprintCallable, Category = "Audio|OSC")
-	static void AddBundleToBundle(const FOSCBundle& inBundle, UPARAM(ref) FOSCBundle& outBundle);
+	static void ClearMessage(UPARAM(ref) FOSCMessage& Message);
 
 	UFUNCTION(BlueprintCallable, Category = "Audio|OSC")
-	static void GetMessagesFromBundle(const FOSCBundle& Bundle, TArray<FOSCMessage>& Messages);
+	static void ClearBundle(UPARAM(ref) FOSCBundle& Bundle);
 
 	UFUNCTION(BlueprintCallable, Category = "Audio|OSC")
-	static void ClearMessage(UPARAM(ref) FOSCMessage& Message);
+	static UPARAM(DisplayName = "Message") FOSCMessage& AddFloat(UPARAM(ref) FOSCMessage& Message, float Value);
 
 	UFUNCTION(BlueprintCallable, Category = "Audio|OSC")
-	static void ClearBundle(UPARAM(ref) FOSCBundle& Bundle);
+	static UPARAM(DisplayName = "Message") FOSCMessage& AddInt32(UPARAM(ref) FOSCMessage& Message, int32 Value);
+
+	UFUNCTION(BlueprintCallable, Category = "Audio|OSC")
+	static UPARAM(DisplayName = "Message") FOSCMessage& AddInt64(UPARAM(ref) FOSCMessage& Message, int64 Value);
+
+	UFUNCTION(BlueprintCallable, Category = "Audio|OSC")
+	static UPARAM(DisplayName = "Message") FOSCMessage& AddString(UPARAM(ref) FOSCMessage& Message, FString Value);
+
+	UFUNCTION(BlueprintCallable, Category = "Audio|OSC")
+	static UPARAM(DisplayName = "Message") FOSCMessage& AddBlob(UPARAM(ref) FOSCMessage& Message, TArray<uint8>& Value);
+
+	UFUNCTION(BlueprintCallable, Category = "Audio|OSC")
+	static UPARAM(DisplayName = "Message") FOSCMessage& AddBool(UPARAM(ref) FOSCMessage& Message, bool Value);
+
+	UFUNCTION(BlueprintCallable, Category = "Audio|OSC")
+	static void GetFloat(const FOSCMessage& Message, const int32 Index, float& Value);
 
 	UFUNCTION(BlueprintCallable, Category = "Audio|OSC")
-	static void SetAddress(UPARAM(ref) FOSCMessage& Message, const FString& Address);
+	static void GetAllFloats(const FOSCMessage& Message, TArray<float>& Values);
 
 	UFUNCTION(BlueprintCallable, Category = "Audio|OSC")
-	static void AddFloat(UPARAM(ref) FOSCMessage& Message, float Value);
+	static void GetInt32(const FOSCMessage& Message, const int32 Index, int32& Value);
 
 	UFUNCTION(BlueprintCallable, Category = "Audio|OSC")
-	static void AddInt32(UPARAM(ref) FOSCMessage& Message, int32 Value);
+	static void GetAllInt32s(const FOSCMessage& Message, TArray<int32>& Values);
 
 	UFUNCTION(BlueprintCallable, Category = "Audio|OSC")
-	static void AddInt64(UPARAM(ref) FOSCMessage& Message, int64 Value);
+	static void GetInt64(const FOSCMessage& Message, const int32 Index, int64& Value);
 
 	UFUNCTION(BlueprintCallable, Category = "Audio|OSC")
-	static void AddString(UPARAM(ref) FOSCMessage& Message, FString Value);
+	static void GetAllInt64s(const FOSCMessage& Message, TArray<int64>& Values);
 
 	UFUNCTION(BlueprintCallable, Category = "Audio|OSC")
-	static void AddBlob(UPARAM(ref) FOSCMessage& Message, TArray<uint8>& Value);
+	static void GetString(const FOSCMessage& Message, const int32 Index, FString& Value);
 
 	UFUNCTION(BlueprintCallable, Category = "Audio|OSC")
-	static void AddBool(UPARAM(ref) FOSCMessage& Message, bool Value);
+	static void GetAllStrings(const FOSCMessage& Message, TArray<FString>& Values);
 
-	//////////////////////////////////////////////////////////////////////////
-	// Audio|OSC functions
 	UFUNCTION(BlueprintCallable, Category = "Audio|OSC")
-	static void GetFloat(const FOSCMessage& Message, const int index, float& Value);
+	static void GetBool(const FOSCMessage& Message, const int32 Index, bool& Value);
 
 	UFUNCTION(BlueprintCallable, Category = "Audio|OSC")
-	static void GetAllFloat(const FOSCMessage& Message, TArray<float>& Values);
+	static void GetAllBools(const FOSCMessage& Message, TArray<bool>& Values);
 
 	UFUNCTION(BlueprintCallable, Category = "Audio|OSC")
-	static void GetInt32(const FOSCMessage& Message, const int index, int32& Value);
+	static void GetBlob(const FOSCMessage& Message, const int32 Index, TArray<uint8>& Value);
 
 	UFUNCTION(BlueprintCallable, Category = "Audio|OSC")
-	static void GetAllInt32(const FOSCMessage& Message, TArray<int32>& Values);
+	static bool OSCAddressIsBundle(const FOSCAddress& Address);
 
 	UFUNCTION(BlueprintCallable, Category = "Audio|OSC")
-	static void GetInt64(const FOSCMessage& Message, const int index, int64& Value);
+	static bool OSCAddressIsMessage(const FOSCAddress& Address);
 
 	UFUNCTION(BlueprintCallable, Category = "Audio|OSC")
-	static void GetAllInt64(const FOSCMessage& Message, TArray<int64>& Values);
+	static bool OSCAddressIsValid(const FOSCAddress& Address);
 
 	UFUNCTION(BlueprintCallable, Category = "Audio|OSC")
-	static void GetString(const FOSCMessage& Message, const int index, FString& Value);
+	static FOSCAddress StringToOSCAddress(const FString& String);
 
 	UFUNCTION(BlueprintCallable, Category = "Audio|OSC")
-	static void GetAllString(const FOSCMessage& Message, TArray<FString>& Values);
+	static UPARAM(DisplayName = "Address") FOSCAddress& OSCAddressAppend(UPARAM(ref) FOSCAddress& Address, const FString& ToAppend);
 
 	UFUNCTION(BlueprintCallable, Category = "Audio|OSC")
-	static void GetBool(const FOSCMessage& Message, const int index, bool& Value);
+	static FOSCAddress GetOSCMessageAddress(const FOSCMessage& Message);
 
 	UFUNCTION(BlueprintCallable, Category = "Audio|OSC")
-	static void GetAllBool(const FOSCMessage& Message, TArray<bool>& Values);
+	static UPARAM(DisplayName = "Message") FOSCMessage& SetOSCMessageAddress(UPARAM(ref) FOSCMessage& Message, const FOSCAddress& Address);
 
 	UFUNCTION(BlueprintCallable, Category = "Audio|OSC")
-	static void GetBlob(const FOSCMessage& Message, const int index, TArray<uint8>& Value);
+	static TArray<FString> SplitOSCAddress(const FOSCAddress& Address);
 };
diff --git a/Engine/Plugins/Runtime/OSC/Source/OSC/Public/OSCMessage.h b/Engine/Plugins/Runtime/OSC/Source/OSC/Public/OSCMessage.h
new file mode 100644
--- /dev/null
+++ b/Engine/Plugins/Runtime/OSC/Source/OSC/Public/OSCMessage.h
@@ -0,0 +1,34 @@
+// Copyright 1998-2019 Epic Games, Inc. All Rights Reserved.
+#pragma once
+
+#include "CoreMinimal.h"
+#include "OSCLog.h"
+#include "OSCPacket.h"
+#include "UObject/ObjectMacros.h"
+#include "UObject/Object.h"
+#include "UObject/WeakObjectPtr.h"
+
+#include "OSCMessage.generated.h"
+
+// Forward Declarations
+class FOSCStream;
+
+
+USTRUCT(BlueprintType)
+struct OSC_API FOSCMessage
+{
+	GENERATED_USTRUCT_BODY()
+
+	FOSCMessage();
+	FOSCMessage(const TSharedPtr<FOSCMessagePacket>& InPacket);
+	~FOSCMessage();
+
+	void SetPacket(TSharedPtr<FOSCMessagePacket>& InPacket);
+	const TSharedPtr<FOSCMessagePacket>& GetPacket() const;
+
+	bool SetAddress(const FOSCAddress& InAddress);
+	const FOSCAddress& GetAddress() const;
+
+private:
+	TSharedPtr<FOSCMessagePacket> Packet;
+};
\ No newline at end of file
diff --git a/Engine/Plugins/Runtime/OSC/Source/OSC/Public/OSCModule.cpp b/Engine/Plugins/Runtime/OSC/Source/OSC/Public/OSCModule.cpp
--- a/Engine/Plugins/Runtime/OSC/Source/OSC/Public/OSCModule.cpp
+++ b/Engine/Plugins/Runtime/OSC/Source/OSC/Public/OSCModule.cpp
@@ -1,32 +1,23 @@
 // Copyright 1998-2019 Epic Games, Inc. All Rights Reserved.
 #include "CoreMinimal.h"
+
 #include "Modules/ModuleInterface.h"
 #include "Modules/ModuleManager.h"
 #include "OSCLog.h"
 
-DEFINE_LOG_CATEGORY( OSCLog );
+DEFINE_LOG_CATEGORY(LogOSC);
+
 
-/**
- * The public interface to this module
- */
 class FOSCModule : public IModuleInterface
 {
 public:
-	virtual void StartupModule() override;
-	virtual void ShutdownModule() override;
-};
-
-IMPLEMENT_MODULE( FOSCModule, OSC )
-
-void FOSCModule::StartupModule()
-{
-	if (!FModuleManager::Get().LoadModule(TEXT("Networking")))
+	virtual void StartupModule()
 	{
-		UE_LOG(OSCLog, Error, TEXT("Required module Networking failed to load"));
-		return;
+		if (!FModuleManager::Get().LoadModule(TEXT("Networking")))
+		{
+			UE_LOG(LogOSC, Error, TEXT("Required module 'Networking' failed to load. OSC service disabled."));
+		}
 	}
-}
+};
 
-void FOSCModule::ShutdownModule()
-{
-}
+IMPLEMENT_MODULE(FOSCModule, OSC)
diff --git a/Engine/Plugins/Runtime/OSC/Source/OSC/Public/OSCServer.h b/Engine/Plugins/Runtime/OSC/Source/OSC/Public/OSCServer.h
new file mode 100644
--- /dev/null
+++ b/Engine/Plugins/Runtime/OSC/Source/OSC/Public/OSCServer.h
@@ -0,0 +1,118 @@
+// Copyright 1998-2019 Epic Games, Inc. All Rights Reserved.
+#pragma once
+
+#include "CoreMinimal.h"
+#include "Containers/CircularQueue.h"
+#include "UObject/Object.h"
+#include "Interfaces/IPv4/IPv4Address.h"
+#include "Interfaces/IPv4/IPv4Endpoint.h"
+#include "Common/UdpSocketReceiver.h"
+#include "OSCMessage.h"
+#include "OSCBundle.h"
+#include "OSCServer.generated.h"
+
+DECLARE_DYNAMIC_MULTICAST_DELEGATE_OneParam(FOSCReceivedMessageEvent, const FOSCMessage &, Message);
+DECLARE_DYNAMIC_MULTICAST_DELEGATE_OneParam(FOSCReceivedBundleEvent, const FOSCBundle &, Bundle);
+
+class FSocket;
+
+UCLASS(BlueprintType)
+class OSC_API UOSCServer : public UObject
+{
+	GENERATED_UCLASS_BODY()
+
+	struct OscCircularBuffer
+	{
+	public:
+
+		OscCircularBuffer(uint32 CircularQueueSize)
+			: Buffer(CircularQueueSize)
+		{}
+
+		OscCircularBuffer()
+			: Buffer(1)
+		{ }
+
+		TCircularQueue<TSharedPtr<FOSCPacket>> Buffer;
+	};
+
+public:
+	/** Gets whether or not to loopback if ReceiveIPAddress provided is multicast. */
+	UFUNCTION(BlueprintCallable, Category = "Audio|OSC")
+	bool GetMulticastLoopback() const;
+
+	/** Returns whether server is actively listening to incoming messages. */
+	UFUNCTION(BlueprintCallable, Category = "Audio|OSC")
+	bool IsActive() const;
+
+	/** Sets the IP address and port to listen for OSC data. */
+	UFUNCTION(BlueprintCallable, Category = "Audio|OSC")
+	void Listen();
+
+	/** Set the address and port of server. Fails if server is currently active. */
+	UFUNCTION(BlueprintCallable, Category = "Audio|OSC")
+	bool SetAddress(const FString& ReceiveIPAddress, int32 Port);
+
+	/** Set whether or not to loopback if ReceiveIPAddress provided is multicast. */
+	UFUNCTION(BlueprintCallable, Category = "Audio|OSC")
+	void SetMulticastLoopback(bool bMulticastLoopback);
+
+	/** Stop and tidy up network socket. */
+	UFUNCTION(BlueprintCallable, Category = "Audio|OSC")
+	void Stop();
+
+	/** Event that gets called when an OSC message is received. */
+	UPROPERTY(BlueprintAssignable, Category = "Audio|OSC")
+	FOSCReceivedMessageEvent OnOscReceived;
+
+	/** Event that gets called when an OSC bundle is received. */
+	UPROPERTY(BlueprintAssignable, Category = "Audio|OSC")
+	FOSCReceivedBundleEvent OnOscBundleReceived;
+	
+	/** When true, only process received events from whitelisted clients. */
+	UPROPERTY(BlueprintReadWrite, Category = "Audio|OSC")
+	bool bWhitelistClients;
+
+	UFUNCTION(BlueprintCallable, Category = "Audio|OSC")
+	void AddWhitelistedClient(const FString& IPAddress);
+
+	UFUNCTION(BlueprintCallable, Category = "Audio|OSC")
+	void RemoveWhitelistedClient(const FString& IPAddress);
+
+	UFUNCTION(BlueprintCallable, Category = "Audio|OSC")
+	void ClearWhitelistedClients();
+
+	UFUNCTION(BlueprintCallable, Category = "Audio|OSC")
+	TSet<FString> GetWhitelistedClients() const;
+
+protected:
+
+	void BeginDestroy() override;
+	
+	/** Callback that receives data from a socket. */
+	void Callback(const FArrayReaderPtr& Data, const FIPv4Endpoint& Endpoint);
+	
+	/** Circular buffer that stores OSC packets. */
+	OscCircularBuffer OSCPackets;
+	
+	/** Socket used to listen for OSC packets. */
+	FSocket* Socket;
+
+	/** UDP receiver. */
+	FUdpSocketReceiver* SocketReceiver;
+
+	/** Set of client addresses whitelisted to process packets from. */
+	TSet<uint32> ClientWhitelist;
+
+	/** OSC buffer lock. */
+	int32 OSCBufferLock;
+
+	/** IPAddress to listen for OSC packets on.  If unset, defaults to LocalHost */
+	FIPv4Address ReceiveIPAddress;
+
+	/** Port to listen for OSC packets on. */
+	int32 Port;
+
+	/** Whether or not to loopback if address provided is multicast */
+	bool bMulticastLoopback;
+};
diff --git a/Engine/Plugins/Runtime/OSC/Source/OSC/Public/OSCStream.h b/Engine/Plugins/Runtime/OSC/Source/OSC/Public/OSCStream.h
new file mode 100644
--- /dev/null
+++ b/Engine/Plugins/Runtime/OSC/Source/OSC/Public/OSCStream.h
@@ -0,0 +1,90 @@
+// Copyright 1998-2019 Epic Games, Inc. All Rights Reserved.
+#pragma once
+
+#include "CoreMinimal.h"
+
+class FOSCStream
+{
+public:
+	FOSCStream();
+	FOSCStream(const uint8* InData, int32 InSize);
+	FOSCStream(int32 InSize);
+
+	/** Returns the stream buffer data */
+	const uint8* GetData() const;
+	
+	/** Returns the stream's length */
+	int32 GetLength() const;
+	
+	/** Returns true if stream has reached the end, false if not */
+	bool HasReachedEnd() const;
+	
+	/** Get current stream position. */
+	int32 GetPosition() const;
+
+	/** Set stream position. */
+	void SetPosition(int32 InPosition);
+
+	/** Read Char from the stream */
+	TCHAR ReadChar();
+
+	/** Write an ansi char into the stream */
+	void WriteChar(TCHAR Char);
+
+	/** Read Int32 from the stream */
+	int32 ReadInt32();
+
+	/** Write Int32 into the stream */
+	void WriteInt32(int32 Value);
+	
+	/** Read Double from the stream */
+	double ReadDouble();
+
+	/** Write Double into the stream */
+	void WriteDouble(uint64 Value);
+	
+	/** Read Int64 from the stream */
+	int64 ReadInt64();
+
+	/** Write Int64 into the stream */
+	void WriteInt64(int64 Value);
+
+	/** Read UInt64 from the stream */
+	uint64 ReadUInt64();
+
+	/** Write UInt64 into the stream */
+	void WriteUInt64(uint64 Value);
+
+	/** Read Float from the stream */
+	float ReadFloat();
+		
+	/** Write Int64 into the stream */
+	void WriteFloat(float Value);
+	
+	/** Read String from the stream */
+	FString ReadString();
+
+	/** Write String into the stream */
+	void WriteString(const FString& String);
+	
+	/** Read blob from the stream */
+	TArray<uint8> ReadBlob();
+
+	/** Write blob into the stream */
+	void WriteBlob(TArray<uint8>& Blob);
+
+private:
+	
+	/** Read data from a provided buffer. */
+	int32 Read(uint8* InBuffer, int32 InSize);
+	
+	/** Write data into buffer. */
+	int32 Write(const uint8* InBuffer, int32 InSize);
+
+	/** Stream data. */
+	TArray<uint8> Data;
+
+	/** Current buffer position. */
+	int32 Position;
+};
+
diff --git a/Engine/Plugins/Runtime/OSC/Source/OSC/Public/OSCTypes.h b/Engine/Plugins/Runtime/OSC/Source/OSC/Public/OSCTypes.h
new file mode 100644
--- /dev/null
+++ b/Engine/Plugins/Runtime/OSC/Source/OSC/Public/OSCTypes.h
@@ -0,0 +1,203 @@
+// Copyright 1998-2019 Epic Games, Inc. All Rights Reserved.
+#pragma once
+
+#include "CoreMinimal.h"
+
+enum class EOSCTypeTag {
+
+	INT32 = 'i',
+	FLOAT = 'f',
+	DOUBLE = 'd',
+	STRING = 's',
+	BLOB = 'b',
+	TIME = 't',
+	INT64 = 'h',
+	CHAR = 'c',
+	TRUE = 'T',
+	FALSE = 'F',
+	NIL = 'N',
+	INFINITUM = 'I',
+	COLOR = 'r',
+	TERMINATE = '\0'
+};
+
+struct FOSCColor
+{
+	explicit FOSCColor() 
+	{}
+
+	explicit FOSCColor(uint8 R, uint8 G, uint8 B, uint8 A) 
+		: Red(R), Green(G), Blue(B), Alpha(A)
+	{}
+
+	explicit FOSCColor(int32 c)
+	{
+		Red = (uint8)(c >> 24);
+		Green = (uint8)(c >> 16);
+		Blue = (uint8)(c >> 8),
+		Alpha = (uint8)c;
+	}
+
+	uint8 Red, Green, Blue, Alpha;
+
+	int32 GetInt32() const
+	{
+		return	(static_cast<int32>(Red)) | (static_cast<int32>(Green) << 8) | 
+				(static_cast<int32>(Blue) << 16) | (static_cast<int32>(Alpha) << 24);
+	}
+};
+
+class FOSCType
+{
+public:
+	explicit FOSCType(int32 Value)
+		: TypeTag(EOSCTypeTag::INT32)
+		, Data(Value)
+		, Blob()
+	{
+	}
+	explicit FOSCType(int64 Value)
+		: TypeTag(EOSCTypeTag::INT64)
+		, Data(Value)
+		, Blob()
+	{
+	}
+	explicit FOSCType(ANSICHAR Value)
+		: TypeTag(EOSCTypeTag::CHAR)
+		, Data(Value)
+		, Blob()
+	{
+	}
+	explicit FOSCType(uint64 Value)
+		: TypeTag(EOSCTypeTag::TIME)
+		, Data(Value)
+		, Blob()
+	{
+	}
+	explicit FOSCType(float Value)
+		: TypeTag(EOSCTypeTag::FLOAT)
+		, Data(Value)
+		, Blob()
+	
+	{
+	}
+	explicit FOSCType(double Value)
+		: TypeTag(EOSCTypeTag::DOUBLE)
+		, Data(Value)
+		, Blob()
+	{
+	}
+	explicit FOSCType(bool Value)
+		: TypeTag(Value ? EOSCTypeTag::TRUE : EOSCTypeTag::FALSE)
+		, Data(Value)
+		, Blob()
+	{
+	}
+	explicit FOSCType(FString Value)
+		: TypeTag(EOSCTypeTag::STRING)
+		, Data(0)
+		, String(Value)
+		, Blob()
+	{
+	}
+	explicit FOSCType(TArray<uint8> Value)
+		: TypeTag(EOSCTypeTag::BLOB)
+		, Data(0)
+		, Blob(Value)
+	{
+	}
+	explicit FOSCType(FOSCColor color)
+		: TypeTag(EOSCTypeTag::COLOR)
+		, Data(0)
+		, Blob()
+		, Color(color)
+	{
+	}
+
+	explicit FOSCType(EOSCTypeTag typeTag)
+		: TypeTag(typeTag)
+		, Data(0)
+		, Blob()
+	{
+	}
+
+	EOSCTypeTag GetTypeTag() const { return TypeTag; }
+
+	bool IsInt32() const { return TypeTag == EOSCTypeTag::INT32; }
+	int32 GetInt32() const { return Data.Int32; }
+
+	bool IsInt64() const { return TypeTag == EOSCTypeTag::INT64; }
+	int64 GetInt64() const { return Data.Int64; }
+
+	bool IsTimeTag() const { return TypeTag == EOSCTypeTag::TIME; }
+	uint64 GetTimeTag() const { return Data.Time; }
+
+	bool IsBool() const { return TypeTag == EOSCTypeTag::TRUE || TypeTag == EOSCTypeTag::FALSE; }
+	bool GetBool() const { return Data.Bool; }
+
+	bool IsChar() const { return TypeTag == EOSCTypeTag::CHAR; }
+	ANSICHAR GetChar() const { return Data.Char; }
+
+	bool IsFloat() const { return TypeTag == EOSCTypeTag::FLOAT; }
+	float GetFloat() const { return Data.Float; }
+	
+	bool IsDouble() const { return TypeTag == EOSCTypeTag::DOUBLE; }
+	double GetDouble() const { return Data.Double; }
+
+	bool IsString() const { return TypeTag == EOSCTypeTag::STRING; }
+	FString GetString() const { return String; }
+
+	bool IsBlob() const { return TypeTag == EOSCTypeTag::BLOB; }
+	TArray<uint8> GetBlob() const { return Blob; }
+
+	bool IsColor() const { return TypeTag == EOSCTypeTag::COLOR; }
+	FOSCColor GetColor()  const { return Color; }
+
+	bool IsNil() const { return TypeTag == EOSCTypeTag::NIL; }
+	bool IsInfinitum() const { return TypeTag == EOSCTypeTag::INFINITUM; }
+
+private:
+	union DataTypes
+	{
+		explicit DataTypes(int32 Value)
+			: Int32(Value)
+		{
+		}
+		explicit DataTypes(int64 Value)
+			: Int64(Value)
+		{
+		}
+		explicit DataTypes(uint64 Value)
+			: Time(Value)
+		{
+		}
+		explicit DataTypes(ANSICHAR Value)
+			: Char(Value)
+		{
+		}
+		explicit DataTypes(float Value)
+			: Float(Value)
+		{
+		}
+		explicit DataTypes(double Value)
+			: Double(Value)
+		{
+		}
+
+		int32 Int32;
+		int64 Int64;
+		uint64 Time;
+		ANSICHAR Char;
+		float Float;
+		double Double;
+		bool Bool;
+	};
+
+	
+private:
+	EOSCTypeTag TypeTag;
+	DataTypes Data;
+	FString String;
+	TArray<uint8> Blob;
+	FOSCColor Color;
+};
