Remove OptimizeCode in plugin

diff --git a/Engine/Plugins/Runtime/OSC/Source/OSC/OSC.Build.cs b/Engine/Plugins/Runtime/OSC/Source/OSC/OSC.Build.cs
--- a/Engine/Plugins/Runtime/OSC/Source/OSC/OSC.Build.cs
+++ b/Engine/Plugins/Runtime/OSC/Source/OSC/OSC.Build.cs
@@ -6,8 +6,6 @@ namespace UnrealBuildTool.Rules
 	{
         public OSC(ReadOnlyTargetRules Target) : base(Target)
 		{
-            OptimizeCode = CodeOptimization.Never;
-            
             PublicDependencyModuleNames.AddRange(
 				new string[] {
                     "Core",
