Fix clang

diff --git a/Engine/Plugins/Runtime/OSC/Source/OSC/Private/OSCBundle.cpp b/Engine/Plugins/Runtime/OSC/Source/OSC/Private/OSCBundle.cpp
--- a/Engine/Plugins/Runtime/OSC/Source/OSC/Private/OSCBundle.cpp
+++ b/Engine/Plugins/Runtime/OSC/Source/OSC/Private/OSCBundle.cpp
@@ -1,5 +1,7 @@
 // Copyright 1998-2019 Epic Games, Inc. All Rights Reserved.
 #include "OSCBundle.h"
+
+#include "OSCAddress.h"
 #include "OSCLog.h"
 #include "OSCStream.h"
 
@@ -84,6 +86,6 @@ bool FOSCBundlePacket::IsMessage()
 
 const FOSCAddress& FOSCBundlePacket::GetAddress() const
 {
-	const static FOSCAddress BundleIdentifier = TEXT("#bundle");
+	const static FOSCAddress BundleIdentifier = FOSCAddress(OSC::BundleTag);
 	return BundleIdentifier;
 }
\ No newline at end of file
