- Remove unnecessary custom types.
- Move Manager implementation to private directory

diff --git a/Engine/Plugins/Runtime/OSC/Source/OSC/Private/OSCManager.cpp b/Engine/Plugins/Runtime/OSC/Source/OSC/Private/OSCManager.cpp
new file mode 100644
--- /dev/null
+++ b/Engine/Plugins/Runtime/OSC/Source/OSC/Private/OSCManager.cpp
@@ -0,0 +1,374 @@
+// Copyright 1998-2019 Epic Games, Inc. All Rights Reserved.
+#include "OSCManager.h"
+
+#include "IPAddress.h"
+#include "OSCMessage.h"
+#include "OSCBundle.h"
+#include "OSCLog.h"
+#include "SocketSubsystem.h"
+
+
+namespace
+{
+	// Returns true if provided address was null and was able to
+	// override with local host address, false if not.
+	bool GetLocalHostAddress(FString& Address)
+	{
+		if (!Address.IsEmpty() && Address != TEXT("0"))
+		{
+			return false;
+		}
+
+		bool bCanBind = false;
+		bool bAppendPort = false;
+		if (ISocketSubsystem* SocketSys = ISocketSubsystem::Get(PLATFORM_SOCKETSUBSYSTEM))
+		{
+			const TSharedPtr<FInternetAddr> Addr = SocketSys->GetLocalHostAddr(*GLog, bCanBind);
+			if (Addr.IsValid())
+			{
+				Address = Addr->ToString(bAppendPort);
+				return true;
+			}
+		}
+
+		return false;
+	}
+
+	const FOSCType* GetOSCTypeAtIndex(const FOSCMessage& Message, const int32 Index)
+	{
+		const TSharedPtr<FOSCMessagePacket>& Packet = Message.GetPacket();
+		if (Packet.IsValid())
+		{
+			TArray<FOSCType>& Args = Packet->GetArguments();
+			if (Index >= Args.Num())
+			{
+				UE_LOG(LogOSC, Error, TEXT("Index '%d' out-of-bounds.  Message argument size = '%d'"), Index, Args.Num());
+				return nullptr;
+			}
+
+			return &Args[Index];
+		}
+
+		return nullptr;
+	}
+} // namespace <>
+
+UOSCServer* UOSCManager::CreateOSCServer(FString ReceiveIPAddress, int32 Port, bool bMulticastLoopback, bool bStartListening)
+{
+	UOSCServer* NewOSCServer = NewObject<UOSCServer>();
+
+	if (GetLocalHostAddress(ReceiveIPAddress))
+	{
+		UE_LOG(LogOSC, Display, TEXT("OSCServer ReceiveAddress not specified. Using LocalHost IP: '%s'"), *ReceiveIPAddress);
+	}
+
+	NewOSCServer->SetMulticastLoopback(bMulticastLoopback);
+	if (NewOSCServer->SetAddress(ReceiveIPAddress, Port))
+	{
+		if (bStartListening)
+		{
+			NewOSCServer->Listen();
+		}
+	}
+	else
+	{
+		UE_LOG(LogOSC, Warning, TEXT("Failed to parse ReceiveAddress '%s' for OSCServer."), *ReceiveIPAddress);
+	}
+
+	return NewOSCServer;
+}
+
+UOSCClient* UOSCManager::CreateOSCClient(FString SendIPAddress, int32 Port)
+{
+	UOSCClient* NewOSCClient = NewObject<UOSCClient>();
+
+	if (GetLocalHostAddress(SendIPAddress))
+	{
+		UE_LOG(LogOSC, Display, TEXT("OSCClient SendAddress not specified. Using LocalHost IP: '%s'"), *SendIPAddress);
+	}
+
+	if (!NewOSCClient->SetSendIPAddress(SendIPAddress, Port))
+	{
+		UE_LOG(LogOSC, Warning, TEXT("Failed to parse SendAddress '%s' for OSCClient. Client unable to send new messages."), *SendIPAddress);
+	}
+
+	return NewOSCClient;
+}
+
+void UOSCManager::ClearMessage(FOSCMessage& Message)
+{
+	if (Message.GetPacket().IsValid())
+	{
+		Message.GetPacket()->GetArguments().Reset();
+	}
+}
+
+void UOSCManager::ClearBundle(FOSCBundle& Bundle)
+{
+	if (Bundle.GetPacket().IsValid())
+	{
+		Bundle.GetPacket()->GetPackets().Reset();
+	}
+}
+
+FOSCBundle& UOSCManager::AddMessageToBundle(const FOSCMessage& Message, FOSCBundle& Bundle)
+{
+	if (Message.GetPacket().IsValid() && Bundle.GetPacket().IsValid())
+	{
+		Bundle.GetPacket()->GetPackets().Add(Message.GetPacket());
+	}
+
+	return Bundle;
+}
+
+FOSCBundle& UOSCManager::AddBundleToBundle(const FOSCBundle& InBundle, FOSCBundle& OutBundle)
+{
+	if (InBundle.GetPacket().IsValid() && OutBundle.GetPacket().IsValid())
+	{
+		OutBundle.GetPacket()->GetPackets().Add(InBundle.GetPacket());
+	}
+
+	return OutBundle;
+}
+
+FOSCMessage& UOSCManager::AddFloat(FOSCMessage& Message, float Value)
+{
+	Message.GetPacket()->GetArguments().Add(FOSCType(Value));
+	return Message;
+}
+
+FOSCMessage& UOSCManager::AddInt32(FOSCMessage& Message, int32 Value)
+{
+	Message.GetPacket()->GetArguments().Add(FOSCType(Value));
+	return Message;
+}
+
+FOSCMessage& UOSCManager::AddInt64(FOSCMessage& Message, int64 Value)
+{
+	Message.GetPacket()->GetArguments().Add(FOSCType(Value));
+	return Message;
+}
+
+FOSCMessage& UOSCManager::AddString(FOSCMessage& Message, FString Value)
+{
+	Message.GetPacket()->GetArguments().Add(FOSCType(Value));
+	return Message;
+}
+
+FOSCMessage& UOSCManager::AddBlob(FOSCMessage& Message, TArray<uint8>& Value)
+{
+	Message.GetPacket()->GetArguments().Add(FOSCType(Value));
+	return Message;
+}
+
+FOSCMessage& UOSCManager::AddBool(FOSCMessage& Message, bool Value)
+{
+	Message.GetPacket()->GetArguments().Add(FOSCType(Value));
+	return Message;
+}
+
+void UOSCManager::GetMessagesFromBundle(const FOSCBundle& Bundle, TArray<FOSCMessage>& Messages)
+{
+	if (Bundle.GetPacket().IsValid())
+	{
+		const TSharedPtr<FOSCBundlePacket> BundlePacket = Bundle.GetPacket();
+		for (int32 i = 0; i < BundlePacket->GetPackets().Num(); i++)
+		{
+			const TSharedPtr<FOSCPacket>& Packet = BundlePacket->GetPackets()[i];
+			if (Packet->IsMessage())
+			{
+				Messages.Emplace(StaticCastSharedPtr<FOSCMessagePacket>(Packet));
+			}
+		}
+	}
+}
+
+void UOSCManager::GetFloat(const FOSCMessage& Message, const int32 Index, float& Value)
+{
+	if (const FOSCType* OSCType = GetOSCTypeAtIndex(Message, Index))
+	{
+		if (OSCType->IsFloat())
+		{
+			Value = OSCType->GetFloat();
+		}
+	}
+}
+
+void UOSCManager::GetAllFloats(const FOSCMessage& Message, TArray<float>& Values)
+{
+	if (Message.GetPacket().IsValid())
+	{
+		const TArray<FOSCType>& Args = Message.GetPacket()->GetArguments();
+		for (int32 i = 0; i < Args.Num(); i++)
+		{
+			const FOSCType& OSCType = Args[i];
+			if (OSCType.IsFloat())
+			{
+				Values.Add(OSCType.GetFloat());
+			}
+		}
+	}
+}
+
+void UOSCManager::GetInt32(const FOSCMessage& Message, const int32 Index, int32& Value)
+{
+	if (const FOSCType* OSCType = GetOSCTypeAtIndex(Message, Index))
+	{
+		if (OSCType->IsInt32())
+		{
+			Value = OSCType->GetInt32();
+		}
+	}
+}
+
+void UOSCManager::GetAllInt32s(const FOSCMessage& Message, TArray<int32>& Values)
+{
+	if (Message.GetPacket().IsValid())
+	{
+		const TArray<FOSCType>& Args = Message.GetPacket()->GetArguments();
+		for (int32 i = 0; i < Args.Num(); i++)
+		{
+			const FOSCType& OSCType = Args[i];
+			if (OSCType.IsInt32())
+			{
+				Values.Add(OSCType.GetInt32());
+			}
+		}
+	}
+}
+
+void UOSCManager::GetInt64(const FOSCMessage& Message, const int32 Index, int64& Value)
+{
+	if (const FOSCType* OSCType = GetOSCTypeAtIndex(Message, Index))
+	{
+		if (OSCType->IsInt64())
+		{
+			Value = OSCType->GetInt64();
+		}
+	}
+}
+
+void UOSCManager::GetAllInt64s(const FOSCMessage& Message, TArray<int64>& Values)
+{
+	if (Message.GetPacket().IsValid())
+	{
+		const TArray<FOSCType>& Args = Message.GetPacket()->GetArguments();
+		for (int32 i = 0; i < Args.Num(); i++)
+		{
+			const FOSCType& OSCType = Args[i];
+			if (OSCType.IsInt64())
+			{
+				Values.Add(OSCType.GetInt64());
+			}
+		}
+	}
+}
+
+void UOSCManager::GetString(const FOSCMessage& Message, const int32 Index, FString& Value)
+{
+	if (const FOSCType* OSCType = GetOSCTypeAtIndex(Message, Index))
+	{
+		if (OSCType->IsString())
+		{
+			Value = OSCType->GetString();
+		}
+	}
+}
+
+void UOSCManager::GetAllStrings(const FOSCMessage& Message, TArray<FString>& Values)
+{
+	if (Message.GetPacket().IsValid())
+	{
+		const TArray<FOSCType>& Args = Message.GetPacket()->GetArguments();
+		for (int32 i = 0; i < Args.Num(); i++)
+		{
+			const FOSCType& OSCType = Args[i];
+			if (OSCType.IsString())
+			{
+				Values.Add(OSCType.GetString());
+			}
+		}
+	}
+}
+
+void UOSCManager::GetBool(const FOSCMessage& Message, const int32 Index, bool& Value)
+{
+	if (const FOSCType* OSCType = GetOSCTypeAtIndex(Message, Index))
+	{
+		if (OSCType->IsBool())
+		{
+			Value = OSCType->GetBool();
+		}
+	}
+}
+
+void UOSCManager::GetAllBools(const FOSCMessage& Message, TArray<bool>& Values)
+{
+	if (Message.GetPacket().IsValid())
+	{
+		const TArray<FOSCType>& Args = Message.GetPacket()->GetArguments();
+		for (int32 i = 0; i < Args.Num(); i++)
+		{
+			const FOSCType& OSCType = Args[i];
+			if (OSCType.IsBool())
+			{
+				Values.Add(OSCType.GetBool());
+			}
+		}
+	}
+}
+
+void UOSCManager::GetBlob(const FOSCMessage& Message, const int32 Index, TArray<uint8>& Value)
+{
+	if (const FOSCType* OSCType = GetOSCTypeAtIndex(Message, Index))
+	{
+		if (OSCType->IsBlob())
+		{
+			Value = OSCType->GetBlob();
+		}
+	}
+}
+
+bool UOSCManager::OSCAddressIsBundle(const FOSCAddress& Address)
+{
+	return Address.IsBundle();
+}
+
+bool UOSCManager::OSCAddressIsMessage(const FOSCAddress& Address)
+{
+	return Address.IsBundle();
+}
+
+bool UOSCManager::OSCAddressIsValid(const FOSCAddress& Address)
+{
+	return Address.IsValid();
+}
+
+FOSCAddress UOSCManager::StringToOSCAddress(const FString& String)
+{
+	return FOSCAddress(String);
+}
+
+FOSCAddress& UOSCManager::OSCAddressAppend(FOSCAddress& Address, const FString& ToAppend)
+{
+	Address.Append(ToAppend);
+	return Address;
+}
+
+FOSCAddress UOSCManager::GetOSCMessageAddress(const FOSCMessage& Message)
+{
+	return Message.GetAddress();
+}
+
+FOSCMessage& UOSCManager::SetOSCMessageAddress(FOSCMessage& Message, const FOSCAddress& Address)
+{
+	Message.SetAddress(Address);
+	return Message;
+}
+
+TArray<FString> UOSCManager::SplitOSCAddress(const FOSCAddress& Address)
+{
+	TArray<FString> OutArray;
+	Address.Split(OutArray);
+	return MoveTemp(OutArray);
+}
\ No newline at end of file
diff --git a/Engine/Plugins/Runtime/OSC/Source/OSC/Private/OSCPacket.cpp b/Engine/Plugins/Runtime/OSC/Source/OSC/Private/OSCPacket.cpp
--- a/Engine/Plugins/Runtime/OSC/Source/OSC/Private/OSCPacket.cpp
+++ b/Engine/Plugins/Runtime/OSC/Source/OSC/Private/OSCPacket.cpp
@@ -117,7 +117,11 @@ void FOSCMessagePacket::WriteData(FOSCStream& Stream)
 		}
 		break;
 		case EOSCTypeTag::COLOR:
-			Stream.WriteInt32(OSCType.GetColor().GetInt32());
+#if PLATFORM_LITTLE_ENDIAN
+			Stream.WriteInt32(OSCType.GetColor().ToPackedABGR());
+#else
+			Stream.WriteInt32(OSCType.GetColor().ToPackedRGBA());
+#endif
 			break;
 		case EOSCTypeTag::TRUE:
 		case EOSCTypeTag::FALSE:
@@ -185,7 +189,7 @@ void FOSCMessagePacket::ReadData(FOSCStream& Stream)
 			Arguments.Add(FOSCType(Stream.ReadBlob()));
 			break;
 		case EOSCTypeTag::COLOR:
-			Arguments.Add(FOSCType(FOSCColor(Stream.ReadInt32())));
+			Arguments.Add(FOSCType(FColor(Stream.ReadInt32())));
 			break;
 		case EOSCTypeTag::TERMINATE:
 			Stream.ReadChar();
diff --git a/Engine/Plugins/Runtime/OSC/Source/OSC/Private/OSCServer.cpp b/Engine/Plugins/Runtime/OSC/Source/OSC/Private/OSCServer.cpp
--- a/Engine/Plugins/Runtime/OSC/Source/OSC/Private/OSCServer.cpp
+++ b/Engine/Plugins/Runtime/OSC/Source/OSC/Private/OSCServer.cpp
@@ -11,15 +11,17 @@
 #include "OSCBundle.h"
 #include "OSCLog.h"
 
-#define OSC_BUFFER_FREE 0
-#define OSC_BUFFER_LOCKED 1
-
+#define OSC_SERVER_BUFFER_FREE 0
+#define OSC_SERVER_BUFFER_LOCKED 1
 
 UOSCServer::UOSCServer(const FObjectInitializer& ObjectInitializer)
 	: Super(ObjectInitializer) 
 	, bWhitelistClients(false)
-	, OSCPackets(1024)
+	, Socket(nullptr)
+	, SocketReceiver(nullptr)
 	, OSCBufferLock(0)
+	, Port(0)
+	, bMulticastLoopback(false)
 {
 }
 
@@ -37,41 +39,40 @@ void UOSCServer::Callback(const FArrayReaderPtr& Data, const FIPv4Endpoint& Endp
 		FOSCStream Stream = FOSCStream(Data->GetData(), Data->Num());
 
 		Packet->ReadData(Stream);
-		OSCPackets.Buffer.Enqueue(MoveTemp(Packet));
+		OSCPackets.Enqueue(MoveTemp(Packet));
 	}
 
-	if (!FPlatformAtomics::InterlockedCompareExchange(&OSCBufferLock, OSC_BUFFER_LOCKED, OSC_BUFFER_FREE) &&
-		!OSCPackets.Buffer.IsEmpty())
+	if (!FPlatformAtomics::InterlockedCompareExchange(&OSCBufferLock, OSC_SERVER_BUFFER_LOCKED, OSC_SERVER_BUFFER_FREE) &&
+		!OSCPackets.IsEmpty())
 	{
-		check(OSCBufferLock == OSC_BUFFER_LOCKED);
-		FSimpleDelegateGraphTask::CreateAndDispatchWhenReady(
-			FSimpleDelegateGraphTask::FDelegate::CreateLambda([this]() {
-			check(OSCBufferLock == OSC_BUFFER_LOCKED);
-			FPlatformAtomics::InterlockedCompareExchange(&OSCBufferLock, OSC_BUFFER_FREE, OSC_BUFFER_LOCKED);
+		check(OSCBufferLock == OSC_SERVER_BUFFER_LOCKED);
+		FSimpleDelegateGraphTask::CreateAndDispatchWhenReady(FSimpleDelegateGraphTask::FDelegate::CreateLambda([this]()
+		{
+			check(OSCBufferLock == OSC_SERVER_BUFFER_LOCKED);
+			FPlatformAtomics::InterlockedCompareExchange(&OSCBufferLock, OSC_SERVER_BUFFER_FREE, OSC_SERVER_BUFFER_LOCKED);
 				
-				TSharedPtr<FOSCPacket> Packet;
-				while (OSCPackets.Buffer.Dequeue(Packet))
+			TSharedPtr<FOSCPacket> Packet;
+			while (OSCPackets.Dequeue(Packet))
+			{
+				if (Packet->IsMessage())
+				{
+					FOSCMessage Message(FOSCMessage(StaticCastSharedPtr<FOSCMessagePacket>(Packet)));
+					OnOscReceived.Broadcast(Message);
+				}
+				else if (Packet->IsBundle())
+				{
+					FOSCBundle Bundle(StaticCastSharedPtr<FOSCBundlePacket>(Packet));
+					OnOscBundleReceived.Broadcast(Bundle);
+				}
+				else
 				{
-					if (Packet->IsMessage())
-					{
-						FOSCMessage Message(FOSCMessage(StaticCastSharedPtr<FOSCMessagePacket>(Packet)));
-						OnOscReceived.Broadcast(Message);
-					}
-					else if (Packet->IsBundle())
-					{
-						FOSCBundle Bundle(StaticCastSharedPtr<FOSCBundlePacket>(Packet));
-						OnOscBundleReceived.Broadcast(Bundle);
-					}
-					else
-					{
-						UE_LOG(LogOSC, Warning, TEXT("Failed to parse invalid received OSC message. OSCAddress '%s' is invalid."), *Packet->GetAddress().Value);
-					}
+					UE_LOG(LogOSC, Warning, TEXT("Failed to parse invalid received OSC message. OSCAddress '%s' is invalid."), *Packet->GetAddress().Value);
 				}
-				OSCPackets.Buffer.Empty();
-			}),
-			TStatId(),
-			nullptr,
-			ENamedThreads::GameThread);
+			}
+		}),
+		TStatId(),
+		nullptr,
+		ENamedThreads::GameThread);
 	}
 }
 
diff --git a/Engine/Plugins/Runtime/OSC/Source/OSC/Private/OSCStream.cpp b/Engine/Plugins/Runtime/OSC/Source/OSC/Private/OSCStream.cpp
--- a/Engine/Plugins/Runtime/OSC/Source/OSC/Private/OSCStream.cpp
+++ b/Engine/Plugins/Runtime/OSC/Source/OSC/Private/OSCStream.cpp
@@ -64,6 +64,23 @@ void FOSCStream::WriteChar(TCHAR Char)
 	Write(ToWrite, 1);
 }
 
+FColor FOSCStream::ReadColor()
+{
+	uint32 Packed = static_cast<uint32>(ReadInt32());
+	return FColor(Packed);
+}
+
+void FOSCStream::WriteColor(FColor Color)
+{
+#if PLATFORM_LITTLE_ENDIAN
+	uint32 Packed = Color.ToPackedABGR();
+#else // PLATFORM_LITTLE_ENDIAN
+	uint32 Packed = Color.ToPackedRGBA();
+#endif // !PLATFORM_LITTLE_ENDIAN
+
+	WriteInt32(static_cast<int32>(Packed));
+}
+
 int32 FOSCStream::ReadInt32()
 {
 	uint8 Temp[4];
@@ -83,7 +100,7 @@ int32 FOSCStream::ReadInt32()
 		return u.i;
 #else
 		return *(int32*)(Temp);
-#endif
+#endif // !PLATFORM_LITTLE_ENDIAN
 	}
 
 	return 0;
@@ -105,9 +122,9 @@ void FOSCStream::WriteInt32(int32 Value)
 	Temp[2] = u.c[1];
 	Temp[1] = u.c[2];
 	Temp[0] = u.c[3];
-#else
-	*reinterpret_cast<int32*>(Temp) = Value;
-#endif
+#else // PLATFORM_LITTLE_ENDIAN
+	*(int32*)(Temp) = Value;
+#endif // !PLATFORM_LITTLE_ENDIAN
 	Write(Temp, 4);
 }
 
@@ -132,9 +149,9 @@ double FOSCStream::ReadDouble()
 		u.c[7] = Temp[0];
 
 		return u.d;
-#else
+#else // PLATFORM_LITTLE_ENDIAN
 		return *(double*)Temp;
-#endif
+#endif // !PLATFORM_LITTLE_ENDIAN
 
 	}
 
@@ -161,9 +178,9 @@ void FOSCStream::WriteDouble(uint64 Value)
 	Temp[2] = u.c[5];
 	Temp[1] = u.c[6];
 	Temp[0] = u.c[7];
-#else
-	*reinterpret_cast<double*>(Temp) = Value;
-#endif
+#else // PLATFORM_LITTLE_ENDIAN
+	*(double*)(Temp) = Value;
+#endif // !PLATFORM_LITTLE_ENDIAN
 	Write(Temp, 8);
 }
 
@@ -188,9 +205,9 @@ int64 FOSCStream::ReadInt64()
 		u.c[7] = Temp[0];
 
 		return u.i;
-#else
+#else // PLATFORM_LITTLE_ENDIAN
 		return *(int64*)Temp;
-#endif
+#endif // !PLATFORM_LITTLE_ENDIAN
 	}
 
 	return 0;
@@ -216,9 +233,9 @@ void FOSCStream::WriteInt64(int64 Value)
 	Temp[2] = u.c[5];
 	Temp[1] = u.c[6];
 	Temp[0] = u.c[7];
-#else
-	*reinterpret_cast<int64*>(Temp) = Value;
-#endif
+#else // PLATFORM_LITTLE_ENDIAN
+	*(int64*)(Temp) = Value;
+#endif // !PLATFORM_LITTLE_ENDIAN
 	Write(Temp, 8);
 }
 
@@ -243,9 +260,9 @@ uint64 FOSCStream::ReadUInt64()
 		u.c[7] = Temp[0];
 
 		return u.i;
-#else
+#else // PLATFORM_LITTLE_ENDIAN
 		return *(uint64*)Temp;
-#endif
+#endif // !PLATFORM_LITTLE_ENDIAN
 
 	}
 
@@ -272,9 +289,9 @@ void FOSCStream::WriteUInt64(uint64 Value)
 	Temp[2] = u.c[5];
 	Temp[1] = u.c[6];
 	Temp[0] = u.c[7];
-#else
+#else // PLATFORM_LITTLE_ENDIAN
 	*reinterpret_cast<uint64*>(Temp) = Value;
-#endif
+#endif // !PLATFORM_LITTLE_ENDIAN
 	Write(Temp, 8);
 }
 
@@ -295,9 +312,9 @@ float FOSCStream::ReadFloat()
 		u.c[3] = Temp[0];
 
 		return u.f;
-#else
+#else // PLATFORM_LITTLE_ENDIAN
 		return *(float*)Temp;
-#endif
+#endif // !PLATFORM_LITTLE_ENDIAN
 	}
 
 	return 0.0f;
@@ -320,7 +337,7 @@ void FOSCStream::WriteFloat(float Value)
 	Temp[1] = u.c[2];
 	Temp[0] = u.c[3];
 #else // if !PLATFORM_LITTLE_ENDIAN
-	*reinterpret_cast<float*>(Temp) = Value;
+	*(float*)(Temp) = Value;
 #endif // !PLATFORM_LITTLE_ENDIAN
 
 	Write(Temp, 4);
diff --git a/Engine/Plugins/Runtime/OSC/Source/OSC/Public/OSCManager.cpp b/Engine/Plugins/Runtime/OSC/Source/OSC/Public/OSCManager.cpp
deleted file mode 100644
diff --git a/Engine/Plugins/Runtime/OSC/Source/OSC/Public/OSCServer.h b/Engine/Plugins/Runtime/OSC/Source/OSC/Public/OSCServer.h
--- a/Engine/Plugins/Runtime/OSC/Source/OSC/Public/OSCServer.h
+++ b/Engine/Plugins/Runtime/OSC/Source/OSC/Public/OSCServer.h
@@ -2,7 +2,7 @@
 #pragma once
 
 #include "CoreMinimal.h"
-#include "Containers/CircularQueue.h"
+
 #include "UObject/Object.h"
 #include "Interfaces/IPv4/IPv4Address.h"
 #include "Interfaces/IPv4/IPv4Endpoint.h"
@@ -21,21 +21,6 @@ class OSC_API UOSCServer : public UObject
 {
 	GENERATED_UCLASS_BODY()
 
-	struct OscCircularBuffer
-	{
-	public:
-
-		OscCircularBuffer(uint32 CircularQueueSize)
-			: Buffer(CircularQueueSize)
-		{}
-
-		OscCircularBuffer()
-			: Buffer(1)
-		{ }
-
-		TCircularQueue<TSharedPtr<FOSCPacket>> Buffer;
-	};
-
 public:
 	/** Gets whether or not to loopback if ReceiveIPAddress provided is multicast. */
 	UFUNCTION(BlueprintCallable, Category = "Audio|OSC")
@@ -73,28 +58,31 @@ public:
 	UPROPERTY(BlueprintReadWrite, Category = "Audio|OSC")
 	bool bWhitelistClients;
 
+	/** Adds client to whitelist of clients to listen for. */
 	UFUNCTION(BlueprintCallable, Category = "Audio|OSC")
 	void AddWhitelistedClient(const FString& IPAddress);
 
+	/** Removes whitelisted client to listen for. */
 	UFUNCTION(BlueprintCallable, Category = "Audio|OSC")
 	void RemoveWhitelistedClient(const FString& IPAddress);
 
+	/** Clears client whitelist to listen for. */
 	UFUNCTION(BlueprintCallable, Category = "Audio|OSC")
 	void ClearWhitelistedClients();
 
+	/** Returns set of whitelisted clients. */
 	UFUNCTION(BlueprintCallable, Category = "Audio|OSC")
 	TSet<FString> GetWhitelistedClients() const;
 
 protected:
-
 	void BeginDestroy() override;
 	
 	/** Callback that receives data from a socket. */
 	void Callback(const FArrayReaderPtr& Data, const FIPv4Endpoint& Endpoint);
 	
 	/** Circular buffer that stores OSC packets. */
-	OscCircularBuffer OSCPackets;
-	
+	TQueue<TSharedPtr<FOSCPacket>> OSCPackets;
+
 	/** Socket used to listen for OSC packets. */
 	FSocket* Socket;
 
diff --git a/Engine/Plugins/Runtime/OSC/Source/OSC/Public/OSCStream.h b/Engine/Plugins/Runtime/OSC/Source/OSC/Public/OSCStream.h
--- a/Engine/Plugins/Runtime/OSC/Source/OSC/Public/OSCStream.h
+++ b/Engine/Plugins/Runtime/OSC/Source/OSC/Public/OSCStream.h
@@ -31,12 +31,18 @@ public:
 	/** Write an ansi char into the stream */
 	void WriteChar(TCHAR Char);
 
+	/** Write Color into the stream */
+	void WriteColor(FColor Color);
+
+	/** Read Color from the stream */
+	FColor ReadColor();
+
 	/** Read Int32 from the stream */
 	int32 ReadInt32();
 
 	/** Write Int32 into the stream */
 	void WriteInt32(int32 Value);
-	
+
 	/** Read Double from the stream */
 	double ReadDouble();
 
diff --git a/Engine/Plugins/Runtime/OSC/Source/OSC/Public/OSCTypes.h b/Engine/Plugins/Runtime/OSC/Source/OSC/Public/OSCTypes.h
--- a/Engine/Plugins/Runtime/OSC/Source/OSC/Public/OSCTypes.h
+++ b/Engine/Plugins/Runtime/OSC/Source/OSC/Public/OSCTypes.h
@@ -3,6 +3,9 @@
 
 #include "CoreMinimal.h"
 
+#include "Math/Color.h"
+
+
 enum class EOSCTypeTag {
 
 	INT32 = 'i',
@@ -21,31 +24,6 @@ enum class EOSCTypeTag {
 	TERMINATE = '\0'
 };
 
-struct FOSCColor
-{
-	explicit FOSCColor() 
-	{}
-
-	explicit FOSCColor(uint8 R, uint8 G, uint8 B, uint8 A) 
-		: Red(R), Green(G), Blue(B), Alpha(A)
-	{}
-
-	explicit FOSCColor(int32 c)
-	{
-		Red = (uint8)(c >> 24);
-		Green = (uint8)(c >> 16);
-		Blue = (uint8)(c >> 8),
-		Alpha = (uint8)c;
-	}
-
-	uint8 Red, Green, Blue, Alpha;
-
-	int32 GetInt32() const
-	{
-		return	(static_cast<int32>(Red)) | (static_cast<int32>(Green) << 8) | 
-				(static_cast<int32>(Blue) << 16) | (static_cast<int32>(Alpha) << 24);
-	}
-};
 
 class FOSCType
 {
@@ -93,29 +71,29 @@ public:
 		, Blob()
 	{
 	}
-	explicit FOSCType(FString Value)
+	explicit FOSCType(const FString& Value)
 		: TypeTag(EOSCTypeTag::STRING)
 		, Data(0)
 		, String(Value)
 		, Blob()
 	{
 	}
-	explicit FOSCType(TArray<uint8> Value)
+	explicit FOSCType(const TArray<uint8>& Value)
 		: TypeTag(EOSCTypeTag::BLOB)
 		, Data(0)
 		, Blob(Value)
 	{
 	}
-	explicit FOSCType(FOSCColor color)
+	explicit FOSCType(FColor Value)
 		: TypeTag(EOSCTypeTag::COLOR)
 		, Data(0)
 		, Blob()
-		, Color(color)
+		, Color(Value)
 	{
 	}
 
-	explicit FOSCType(EOSCTypeTag typeTag)
-		: TypeTag(typeTag)
+	explicit FOSCType(EOSCTypeTag TypeTag)
+		: TypeTag(TypeTag)
 		, Data(0)
 		, Blob()
 	{
@@ -151,7 +129,7 @@ public:
 	TArray<uint8> GetBlob() const { return Blob; }
 
 	bool IsColor() const { return TypeTag == EOSCTypeTag::COLOR; }
-	FOSCColor GetColor()  const { return Color; }
+	FColor GetColor()  const { return Color; }
 
 	bool IsNil() const { return TypeTag == EOSCTypeTag::NIL; }
 	bool IsInfinitum() const { return TypeTag == EOSCTypeTag::INFINITUM; }
@@ -199,5 +177,5 @@ private:
 	DataTypes Data;
 	FString String;
 	TArray<uint8> Blob;
-	FOSCColor Color;
+	FColor Color;
 };
