Fix earlier errors backed out by moving packet implementation out of public API to private implementation so manager implementation can be private

Reinstating change:
- Remove unnecessary custom types
- Move Manager implementation to private directory

diff --git a/Engine/Plugins/Runtime/OSC/Source/OSC/Private/OSCBundle.cpp b/Engine/Plugins/Runtime/OSC/Source/OSC/Private/OSCBundle.cpp
--- a/Engine/Plugins/Runtime/OSC/Source/OSC/Private/OSCBundle.cpp
+++ b/Engine/Plugins/Runtime/OSC/Source/OSC/Private/OSCBundle.cpp
@@ -3,89 +3,30 @@
 
 #include "OSCAddress.h"
 #include "OSCLog.h"
-#include "OSCStream.h"
+#include "OSCBundlePacket.h"
 
 
-FOSCBundlePacket::FOSCBundlePacket()
-	: FOSCPacket()
-	, TimeTag(0)
+FOSCBundle::FOSCBundle()
+	: Packet(MakeShareable(new FOSCBundlePacket()))
 {
 }
 
-FOSCBundlePacket::~FOSCBundlePacket()
+FOSCBundle::FOSCBundle(const TSharedPtr<IOSCPacket>& InPacket)
+	: Packet(InPacket)
 {
 }
 
-void FOSCBundlePacket::SetTimeTag(uint64 NewTimeTag) 
-{ 
-	TimeTag = FOSCType(NewTimeTag); 
-}
-
-uint64 FOSCBundlePacket::GetTimeTag() const 
-{ 
-	return TimeTag.GetTimeTag(); 
-}
-
-void FOSCBundlePacket::WriteData(FOSCStream& Stream)
-{
-	if (!GetAddress().IsBundle())
-	{
-		UE_LOG(LogOSC, Warning, TEXT("Failed to write OSCBundlePacket. Invalid OSCAddress '%s'"), *GetAddress().Value);
-		return;
-	}
-
-	// Write bundle & time tag
-	Stream.WriteString(GetAddress().Value);
-	Stream.WriteUInt64(GetTimeTag());
-
-	for (TSharedPtr<FOSCPacket>& Packet : Packets)
-	{
-		int32 StreamPos = Stream.GetPosition();
-		Stream.WriteInt32(0);
-
-		int32 InitPos = Stream.GetPosition();
-		Packet->WriteData(Stream);
-		int32 NewPos = Stream.GetPosition();
-
-		Stream.SetPosition(StreamPos);
-		Stream.WriteInt32(NewPos - InitPos);
-		Stream.SetPosition(NewPos);
-	}
-}
-
-void FOSCBundlePacket::ReadData(FOSCStream& Stream)
+FOSCBundle::~FOSCBundle()
 {
-	TimeTag = FOSCType(Stream.ReadUInt64());
-
-	while (!Stream.HasReachedEnd())
-	{
-		TSharedPtr<FOSCPacket> Packet = FOSCPacket::CreatePacket(Stream.GetData());
-
-		if (Packet.IsValid())
-		{
-			Packet->ReadData(Stream);
-			Packets.Add(Packet);
-		}
-	}
+	Packet.Reset();
 }
 
-FOSCBundlePacket::FPacketBundle& FOSCBundlePacket::GetPackets()
+void FOSCBundle::SetPacket(const TSharedPtr<IOSCPacket>& InPacket)
 {
-	return Packets;
+	Packet = InPacket;
 }
 
-bool FOSCBundlePacket::IsBundle()
+const TSharedPtr<IOSCPacket>& FOSCBundle::GetPacket() const
 {
-	return true;
+	return Packet;
 }
-
-bool FOSCBundlePacket::IsMessage()
-{
-	return false;
-}
-
-const FOSCAddress& FOSCBundlePacket::GetAddress() const
-{
-	const static FOSCAddress BundleIdentifier = FOSCAddress(OSC::BundleTag);
-	return BundleIdentifier;
-}
\ No newline at end of file
diff --git a/Engine/Plugins/Runtime/OSC/Source/OSC/Private/OSCBundlePacket.cpp b/Engine/Plugins/Runtime/OSC/Source/OSC/Private/OSCBundlePacket.cpp
new file mode 100644
--- /dev/null
+++ b/Engine/Plugins/Runtime/OSC/Source/OSC/Private/OSCBundlePacket.cpp
@@ -0,0 +1,96 @@
+// Copyright 1998-2019 Epic Games, Inc. All Rights Reserved.
+#include "OSCBundlePacket.h"
+
+#include "OSCBundle.h"
+#include "OSCLog.h"
+#include "OSCMessage.h"
+#include "OSCTypes.h"
+
+
+FOSCBundlePacket::FOSCBundlePacket()
+	: IOSCPacket()
+	, TimeTag(0)
+{
+}
+
+FOSCBundlePacket::~FOSCBundlePacket()
+{
+}
+
+void FOSCBundlePacket::SetTimeTag(uint64 NewTimeTag)
+{
+	TimeTag = FOSCType(NewTimeTag);
+}
+
+uint64 FOSCBundlePacket::GetTimeTag() const
+{
+	return TimeTag.GetTimeTag();
+}
+
+void FOSCBundlePacket::WriteData(FOSCStream& Stream)
+{
+	if (!GetAddress().IsBundle())
+	{
+		UE_LOG(LogOSC, Warning, TEXT("Failed to write OSCBundlePacket. Invalid OSCAddress '%s'"), *GetAddress().Value);
+		return;
+	}
+
+	// Write bundle & time tag
+	Stream.WriteString(GetAddress().Value);
+	Stream.WriteUInt64(GetTimeTag());
+
+	for (TSharedPtr<IOSCPacket>& Packet : Packets)
+	{
+		int32 StreamPos = Stream.GetPosition();
+		Stream.WriteInt32(0);
+
+		int32 InitPos = Stream.GetPosition();
+		Packet->WriteData(Stream);
+		int32 NewPos = Stream.GetPosition();
+
+		Stream.SetPosition(StreamPos);
+		Stream.WriteInt32(NewPos - InitPos);
+		Stream.SetPosition(NewPos);
+	}
+}
+
+void FOSCBundlePacket::ReadData(FOSCStream& Stream)
+{
+	TimeTag = FOSCType(Stream.ReadUInt64());
+
+	while (!Stream.HasReachedEnd())
+	{
+		TSharedPtr<IOSCPacket> Packet = IOSCPacket::CreatePacket(Stream.GetData());
+
+		if (Packet.IsValid())
+		{
+			Packet->ReadData(Stream);
+			Packets.Add(Packet);
+		}
+	}
+}
+
+FOSCBundlePacket::FPacketBundle& FOSCBundlePacket::GetPackets()
+{
+	return Packets;
+}
+
+bool FOSCBundlePacket::IsBundle()
+{
+	return true;
+}
+
+bool FOSCBundlePacket::IsMessage()
+{
+	return false;
+}
+
+const FOSCAddress& FOSCBundlePacket::GetAddress() const
+{
+	const static FOSCAddress BundleIdentifier = FOSCAddress(OSC::BundleTag);
+	return BundleIdentifier;
+}
+
+void FOSCBundlePacket::SetAddress(const FOSCAddress& InAddress)
+{
+}
diff --git a/Engine/Plugins/Runtime/OSC/Source/OSC/Private/OSCBundlePacket.h b/Engine/Plugins/Runtime/OSC/Source/OSC/Private/OSCBundlePacket.h
new file mode 100644
--- /dev/null
+++ b/Engine/Plugins/Runtime/OSC/Source/OSC/Private/OSCBundlePacket.h
@@ -0,0 +1,47 @@
+// Copyright 1998-2019 Epic Games, Inc. All Rights Reserved.
+#pragma once
+
+#include "OSCAddress.h"
+#include "OSCPacket.h"
+#include "OSCStream.h"
+#include "OSCTypes.h"
+
+
+class FOSCBundlePacket : public IOSCPacket
+{
+public:
+	FOSCBundlePacket();
+	virtual ~FOSCBundlePacket();
+
+	using FPacketBundle = TArray<TSharedPtr<IOSCPacket>>;
+
+	/** Set the bundle time tag. */
+	void SetTimeTag(uint64 NewTimeTag);
+
+	/** Get the bundle time tag. */
+	uint64 GetTimeTag() const;
+
+	/** Get OSC packet by index. */
+	FPacketBundle& GetPackets();
+
+	virtual bool IsBundle() override;
+	virtual bool IsMessage() override;
+
+	/** Writes bundle data into the OSC stream. */
+	virtual void WriteData(FOSCStream& Stream) override;
+
+	/** Reads bundle data from provided OSC stream,
+	  * adding packet data to internal packet bundle. */
+	virtual void ReadData(FOSCStream& Stream) override;
+
+	virtual const FOSCAddress& GetAddress() const override;
+
+	virtual void SetAddress(const FOSCAddress& Address) override;
+
+private:
+	/** Bundle of OSC packets. */
+	FPacketBundle Packets;
+
+	/** Bundle time tag. */
+	FOSCType TimeTag;
+};
diff --git a/Engine/Plugins/Runtime/OSC/Source/OSC/Private/OSCClient.cpp b/Engine/Plugins/Runtime/OSC/Source/OSC/Private/OSCClient.cpp
--- a/Engine/Plugins/Runtime/OSC/Source/OSC/Private/OSCClient.cpp
+++ b/Engine/Plugins/Runtime/OSC/Source/OSC/Private/OSCClient.cpp
@@ -7,6 +7,7 @@
 #include "OSCStream.h"
 #include "OSCBundle.h"
 #include "OSCLog.h"
+#include "OSCPacket.h"
 
 
 namespace
@@ -14,6 +15,7 @@ namespace
 	static const int32 OUTPUT_BUFFER_SIZE = 1024;
 } // namespace <>
 
+
 UOSCClient::UOSCClient()
 	: Socket(FUdpSocketBuilder(*GetName()).Build())
 {
@@ -50,7 +52,7 @@ void UOSCClient::GetSendIPAddress(FString& InIPAddress, int32& Port)
 	Port = IPAddress->GetPort();
 }
 
-void UOSCClient::SendPacket(FOSCPacket& Packet)
+void UOSCClient::SendPacket(IOSCPacket& Packet)
 {
 	const FOSCAddress& OSCAddress = Packet.GetAddress();
 	if (!OSCAddress.IsValid())
@@ -105,14 +107,14 @@ void UOSCClient::BeginDestroy()
 
 void UOSCClient::SendOSCMessage(FOSCMessage& Message)
 {
-	const TSharedPtr<FOSCMessagePacket>& Packet = Message.GetPacket();
+	const TSharedPtr<IOSCPacket>& Packet = Message.GetPacket();
 	check(Packet.IsValid());
 	SendPacket(*Packet.Get());
 }
 
 void UOSCClient::SendOSCBundle(FOSCBundle& Bundle)
 {
-	const TSharedPtr<FOSCBundlePacket>& Packet = Bundle.GetPacket();
+	const TSharedPtr<IOSCPacket>& Packet = Bundle.GetPacket();
 	check(Packet.IsValid());
 	SendPacket(*Packet.Get());
 }
diff --git a/Engine/Plugins/Runtime/OSC/Source/OSC/Private/OSCManager.cpp b/Engine/Plugins/Runtime/OSC/Source/OSC/Private/OSCManager.cpp
new file mode 100644
--- /dev/null
+++ b/Engine/Plugins/Runtime/OSC/Source/OSC/Private/OSCManager.cpp
@@ -0,0 +1,395 @@
+// Copyright 1998-2019 Epic Games, Inc. All Rights Reserved.
+#include "OSCManager.h"
+
+#include "IPAddress.h"
+#include "OSCMessage.h"
+#include "OSCMessagePacket.h"
+#include "OSCBundle.h"
+#include "OSCBundlePacket.h"
+#include "OSCLog.h"
+#include "SocketSubsystem.h"
+
+
+namespace
+{
+	// Returns true if provided address was null and was able to
+	// override with local host address, false if not.
+	bool GetLocalHostAddress(FString& Address)
+	{
+		if (!Address.IsEmpty() && Address != TEXT("0"))
+		{
+			return false;
+		}
+
+		bool bCanBind = false;
+		bool bAppendPort = false;
+		if (ISocketSubsystem* SocketSys = ISocketSubsystem::Get(PLATFORM_SOCKETSUBSYSTEM))
+		{
+			const TSharedPtr<FInternetAddr> Addr = SocketSys->GetLocalHostAddr(*GLog, bCanBind);
+			if (Addr.IsValid())
+			{
+				Address = Addr->ToString(bAppendPort);
+				return true;
+			}
+		}
+
+		return false;
+	}
+
+	const FOSCType* GetOSCTypeAtIndex(const FOSCMessage& Message, const int32 Index)
+	{
+		const TSharedPtr<FOSCMessagePacket>& Packet = StaticCastSharedPtr<FOSCMessagePacket>(Message.GetPacket());
+		if (Packet.IsValid())
+		{
+			TArray<FOSCType>& Args = Packet->GetArguments();
+			if (Index >= Args.Num())
+			{
+				UE_LOG(LogOSC, Error, TEXT("Index '%d' out-of-bounds.  Message argument size = '%d'"), Index, Args.Num());
+				return nullptr;
+			}
+
+			return &Args[Index];
+		}
+
+		return nullptr;
+	}
+} // namespace <>
+
+UOSCServer* UOSCManager::CreateOSCServer(FString ReceiveIPAddress, int32 Port, bool bMulticastLoopback, bool bStartListening)
+{
+	UOSCServer* NewOSCServer = NewObject<UOSCServer>();
+
+	if (GetLocalHostAddress(ReceiveIPAddress))
+	{
+		UE_LOG(LogOSC, Display, TEXT("OSCServer ReceiveAddress not specified. Using LocalHost IP: '%s'"), *ReceiveIPAddress);
+	}
+
+	NewOSCServer->SetMulticastLoopback(bMulticastLoopback);
+	if (NewOSCServer->SetAddress(ReceiveIPAddress, Port))
+	{
+		if (bStartListening)
+		{
+			NewOSCServer->Listen();
+		}
+	}
+	else
+	{
+		UE_LOG(LogOSC, Warning, TEXT("Failed to parse ReceiveAddress '%s' for OSCServer."), *ReceiveIPAddress);
+	}
+
+	return NewOSCServer;
+}
+
+UOSCClient* UOSCManager::CreateOSCClient(FString SendIPAddress, int32 Port)
+{
+	UOSCClient* NewOSCClient = NewObject<UOSCClient>();
+
+	if (GetLocalHostAddress(SendIPAddress))
+	{
+		UE_LOG(LogOSC, Display, TEXT("OSCClient SendAddress not specified. Using LocalHost IP: '%s'"), *SendIPAddress);
+	}
+
+	if (!NewOSCClient->SetSendIPAddress(SendIPAddress, Port))
+	{
+		UE_LOG(LogOSC, Warning, TEXT("Failed to parse SendAddress '%s' for OSCClient. Client unable to send new messages."), *SendIPAddress);
+	}
+
+	return NewOSCClient;
+}
+
+void UOSCManager::ClearMessage(FOSCMessage& Message)
+{
+	const TSharedPtr<FOSCMessagePacket>& Packet = StaticCastSharedPtr<FOSCMessagePacket>(Message.GetPacket());
+	if (Packet.IsValid())
+	{
+		Packet->GetArguments().Reset();
+	}
+}
+
+void UOSCManager::ClearBundle(FOSCBundle& Bundle)
+{
+	const TSharedPtr<FOSCBundlePacket>& Packet = StaticCastSharedPtr<FOSCBundlePacket>(Bundle.GetPacket());
+	if (Packet.IsValid())
+	{
+		Packet->GetPackets().Reset();
+	}
+}
+
+FOSCBundle& UOSCManager::AddMessageToBundle(const FOSCMessage& Message, FOSCBundle& Bundle)
+{
+	const TSharedPtr<FOSCBundlePacket>& BundlePacket = StaticCastSharedPtr<FOSCBundlePacket>(Bundle.GetPacket());
+	const TSharedPtr<FOSCMessagePacket>& MessagePacket = StaticCastSharedPtr<FOSCMessagePacket>(Message.GetPacket());
+
+	if (MessagePacket.IsValid() && BundlePacket.IsValid())
+	{
+		BundlePacket->GetPackets().Add(MessagePacket);
+	}
+
+	return Bundle;
+}
+
+FOSCBundle& UOSCManager::AddBundleToBundle(const FOSCBundle& InBundle, FOSCBundle& OutBundle)
+{
+	const TSharedPtr<FOSCBundlePacket>& InBundlePacket = StaticCastSharedPtr<FOSCBundlePacket>(InBundle.GetPacket());
+	const TSharedPtr<FOSCBundlePacket>& OutBundlePacket = StaticCastSharedPtr<FOSCBundlePacket>(OutBundle.GetPacket());
+
+	if (InBundlePacket.IsValid() && OutBundlePacket.IsValid())
+	{
+		InBundlePacket->GetPackets().Add(OutBundlePacket);
+	}
+
+	return OutBundle;
+}
+
+FOSCMessage& UOSCManager::AddFloat(FOSCMessage& Message, float Value)
+{
+	const TSharedPtr<FOSCMessagePacket>& MessagePacket = StaticCastSharedPtr<FOSCMessagePacket>(Message.GetPacket());
+	MessagePacket->GetArguments().Add(FOSCType(Value));
+	return Message;
+}
+
+FOSCMessage& UOSCManager::AddInt32(FOSCMessage& Message, int32 Value)
+{
+	const TSharedPtr<FOSCMessagePacket>& MessagePacket = StaticCastSharedPtr<FOSCMessagePacket>(Message.GetPacket());
+	MessagePacket->GetArguments().Add(FOSCType(Value));
+	return Message;
+}
+
+FOSCMessage& UOSCManager::AddInt64(FOSCMessage& Message, int64 Value)
+{
+	const TSharedPtr<FOSCMessagePacket>& MessagePacket = StaticCastSharedPtr<FOSCMessagePacket>(Message.GetPacket());
+	MessagePacket->GetArguments().Add(FOSCType(Value));
+	return Message;
+}
+
+FOSCMessage& UOSCManager::AddString(FOSCMessage& Message, FString Value)
+{
+	const TSharedPtr<FOSCMessagePacket>& MessagePacket = StaticCastSharedPtr<FOSCMessagePacket>(Message.GetPacket());
+	MessagePacket->GetArguments().Add(FOSCType(Value));
+	return Message;
+}
+
+FOSCMessage& UOSCManager::AddBlob(FOSCMessage& Message, TArray<uint8>& Value)
+{
+	const TSharedPtr<FOSCMessagePacket>& MessagePacket = StaticCastSharedPtr<FOSCMessagePacket>(Message.GetPacket());
+	MessagePacket->GetArguments().Add(FOSCType(Value));
+	return Message;
+}
+
+FOSCMessage& UOSCManager::AddBool(FOSCMessage& Message, bool Value)
+{
+	const TSharedPtr<FOSCMessagePacket>& MessagePacket = StaticCastSharedPtr<FOSCMessagePacket>(Message.GetPacket());
+	MessagePacket->GetArguments().Add(FOSCType(Value));
+	return Message;
+}
+
+void UOSCManager::GetMessagesFromBundle(const FOSCBundle& Bundle, TArray<FOSCMessage>& Messages)
+{
+	if (Bundle.GetPacket().IsValid())
+	{
+		const TSharedPtr<FOSCBundlePacket>& BundlePacket = StaticCastSharedPtr<FOSCBundlePacket>(Bundle.GetPacket());
+		for (int32 i = 0; i < BundlePacket->GetPackets().Num(); i++)
+		{
+			const TSharedPtr<IOSCPacket>& Packet = BundlePacket->GetPackets()[i];
+			if (Packet->IsMessage())
+			{
+				Messages.Emplace(StaticCastSharedPtr<FOSCMessagePacket>(Packet));
+			}
+		}
+	}
+}
+
+void UOSCManager::GetFloat(const FOSCMessage& Message, const int32 Index, float& Value)
+{
+	if (const FOSCType* OSCType = GetOSCTypeAtIndex(Message, Index))
+	{
+		if (OSCType->IsFloat())
+		{
+			Value = OSCType->GetFloat();
+		}
+	}
+}
+
+void UOSCManager::GetAllFloats(const FOSCMessage& Message, TArray<float>& Values)
+{
+	if (Message.GetPacket().IsValid())
+	{
+		const TSharedPtr<FOSCMessagePacket>& MessagePacket = StaticCastSharedPtr<FOSCMessagePacket>(Message.GetPacket());
+		const TArray<FOSCType>& Args = MessagePacket->GetArguments();
+		for (int32 i = 0; i < Args.Num(); i++)
+		{
+			const FOSCType& OSCType = Args[i];
+			if (OSCType.IsFloat())
+			{
+				Values.Add(OSCType.GetFloat());
+			}
+		}
+	}
+}
+
+void UOSCManager::GetInt32(const FOSCMessage& Message, const int32 Index, int32& Value)
+{
+	if (const FOSCType* OSCType = GetOSCTypeAtIndex(Message, Index))
+	{
+		if (OSCType->IsInt32())
+		{
+			Value = OSCType->GetInt32();
+		}
+	}
+}
+
+void UOSCManager::GetAllInt32s(const FOSCMessage& Message, TArray<int32>& Values)
+{
+	if (Message.GetPacket().IsValid())
+	{
+		const TSharedPtr<FOSCMessagePacket>& MessagePacket = StaticCastSharedPtr<FOSCMessagePacket>(Message.GetPacket());
+		const TArray<FOSCType>& Args = MessagePacket->GetArguments();
+		for (int32 i = 0; i < Args.Num(); i++)
+		{
+			const FOSCType& OSCType = Args[i];
+			if (OSCType.IsInt32())
+			{
+				Values.Add(OSCType.GetInt32());
+			}
+		}
+	}
+}
+
+void UOSCManager::GetInt64(const FOSCMessage& Message, const int32 Index, int64& Value)
+{
+	if (const FOSCType* OSCType = GetOSCTypeAtIndex(Message, Index))
+	{
+		if (OSCType->IsInt64())
+		{
+			Value = OSCType->GetInt64();
+		}
+	}
+}
+
+void UOSCManager::GetAllInt64s(const FOSCMessage& Message, TArray<int64>& Values)
+{
+	const TSharedPtr<FOSCMessagePacket>& MessagePacket = StaticCastSharedPtr<FOSCMessagePacket>(Message.GetPacket());
+	if (MessagePacket.IsValid())
+	{
+		const TArray<FOSCType>& Args = MessagePacket->GetArguments();
+		for (int32 i = 0; i < Args.Num(); i++)
+		{
+			const FOSCType& OSCType = Args[i];
+			if (OSCType.IsInt64())
+			{
+				Values.Add(OSCType.GetInt64());
+			}
+		}
+	}
+}
+
+void UOSCManager::GetString(const FOSCMessage& Message, const int32 Index, FString& Value)
+{
+	if (const FOSCType* OSCType = GetOSCTypeAtIndex(Message, Index))
+	{
+		if (OSCType->IsString())
+		{
+			Value = OSCType->GetString();
+		}
+	}
+}
+
+void UOSCManager::GetAllStrings(const FOSCMessage& Message, TArray<FString>& Values)
+{
+	const TSharedPtr<FOSCMessagePacket>& MessagePacket = StaticCastSharedPtr<FOSCMessagePacket>(Message.GetPacket());
+	if (MessagePacket.IsValid())
+	{
+		const TArray<FOSCType>& Args = MessagePacket->GetArguments();
+		for (int32 i = 0; i < Args.Num(); i++)
+		{
+			const FOSCType& OSCType = Args[i];
+			if (OSCType.IsString())
+			{
+				Values.Add(OSCType.GetString());
+			}
+		}
+	}
+}
+
+void UOSCManager::GetBool(const FOSCMessage& Message, const int32 Index, bool& Value)
+{
+	if (const FOSCType* OSCType = GetOSCTypeAtIndex(Message, Index))
+	{
+		if (OSCType->IsBool())
+		{
+			Value = OSCType->GetBool();
+		}
+	}
+}
+
+void UOSCManager::GetAllBools(const FOSCMessage& Message, TArray<bool>& Values)
+{
+	const TSharedPtr<FOSCMessagePacket>& MessagePacket = StaticCastSharedPtr<FOSCMessagePacket>(Message.GetPacket());
+	if (MessagePacket.IsValid())
+	{
+		const TArray<FOSCType>& Args = MessagePacket->GetArguments();
+		for (int32 i = 0; i < Args.Num(); i++)
+		{
+			const FOSCType& OSCType = Args[i];
+			if (OSCType.IsBool())
+			{
+				Values.Add(OSCType.GetBool());
+			}
+		}
+	}
+}
+
+void UOSCManager::GetBlob(const FOSCMessage& Message, const int32 Index, TArray<uint8>& Value)
+{
+	if (const FOSCType* OSCType = GetOSCTypeAtIndex(Message, Index))
+	{
+		if (OSCType->IsBlob())
+		{
+			Value = OSCType->GetBlob();
+		}
+	}
+}
+
+bool UOSCManager::OSCAddressIsBundle(const FOSCAddress& Address)
+{
+	return Address.IsBundle();
+}
+
+bool UOSCManager::OSCAddressIsMessage(const FOSCAddress& Address)
+{
+	return Address.IsBundle();
+}
+
+bool UOSCManager::OSCAddressIsValid(const FOSCAddress& Address)
+{
+	return Address.IsValid();
+}
+
+FOSCAddress UOSCManager::StringToOSCAddress(const FString& String)
+{
+	return FOSCAddress(String);
+}
+
+FOSCAddress& UOSCManager::OSCAddressAppend(FOSCAddress& Address, const FString& ToAppend)
+{
+	Address.Append(ToAppend);
+	return Address;
+}
+
+FOSCAddress UOSCManager::GetOSCMessageAddress(const FOSCMessage& Message)
+{
+	return Message.GetAddress();
+}
+
+FOSCMessage& UOSCManager::SetOSCMessageAddress(FOSCMessage& Message, const FOSCAddress& Address)
+{
+	Message.SetAddress(Address);
+	return Message;
+}
+
+TArray<FString> UOSCManager::SplitOSCAddress(const FOSCAddress& Address)
+{
+	TArray<FString> OutArray;
+	Address.Split(OutArray);
+	return MoveTemp(OutArray);
+}
\ No newline at end of file
diff --git a/Engine/Plugins/Runtime/OSC/Source/OSC/Private/OSCMessage.cpp b/Engine/Plugins/Runtime/OSC/Source/OSC/Private/OSCMessage.cpp
--- a/Engine/Plugins/Runtime/OSC/Source/OSC/Private/OSCMessage.cpp
+++ b/Engine/Plugins/Runtime/OSC/Source/OSC/Private/OSCMessage.cpp
@@ -1,13 +1,16 @@
 // Copyright 1998-2019 Epic Games, Inc. All Rights Reserved.
 #include "OSCMessage.h"
+
+#include "OSCMessagePacket.h"
 #include "OSCStream.h"
 
+
 FOSCMessage::FOSCMessage()
 	: Packet(MakeShareable(new FOSCMessagePacket()))
 {
 }
 
-FOSCMessage::FOSCMessage(const TSharedPtr<FOSCMessagePacket>& InPacket)
+FOSCMessage::FOSCMessage(const TSharedPtr<IOSCPacket>& InPacket)
 	: Packet(InPacket)
 {
 }
@@ -17,12 +20,12 @@ FOSCMessage::~FOSCMessage()
 	Packet.Reset();
 }
 
-void FOSCMessage::SetPacket(TSharedPtr<FOSCMessagePacket>& InPacket)
+void FOSCMessage::SetPacket(TSharedPtr<IOSCPacket>& InPacket)
 {
 	Packet = InPacket;
 }
 
-const TSharedPtr<FOSCMessagePacket>& FOSCMessage::GetPacket() const
+const TSharedPtr<IOSCPacket>& FOSCMessage::GetPacket() const
 {
 	return Packet;
 }
diff --git a/Engine/Plugins/Runtime/OSC/Source/OSC/Private/OSCMessagePacket.cpp b/Engine/Plugins/Runtime/OSC/Source/OSC/Private/OSCMessagePacket.cpp
new file mode 100644
--- /dev/null
+++ b/Engine/Plugins/Runtime/OSC/Source/OSC/Private/OSCMessagePacket.cpp
@@ -0,0 +1,177 @@
+// Copyright 1998-2019 Epic Games, Inc. All Rights Reserved.
+#include "OSCMessagePacket.h"
+
+#include "OSCBundle.h"
+#include "OSCLog.h"
+#include "OSCMessage.h"
+
+
+FOSCMessagePacket::FOSCMessagePacket()
+	: IOSCPacket()
+{}
+
+FOSCMessagePacket::~FOSCMessagePacket()
+{}
+
+const FOSCAddress& FOSCMessagePacket::GetAddress() const
+{
+	return Address;
+}
+
+void FOSCMessagePacket::SetAddress(const FOSCAddress& InAddress)
+{
+	Address = InAddress;
+}
+
+TArray<FOSCType>& FOSCMessagePacket::GetArguments()
+{
+	return Arguments;
+}
+
+bool FOSCMessagePacket::IsBundle()
+{
+	return false;
+}
+
+bool FOSCMessagePacket::IsMessage()
+{
+	return true;
+}
+
+void FOSCMessagePacket::WriteData(FOSCStream& Stream)
+{
+	if (!Address.IsMessage())
+	{
+		UE_LOG(LogOSC, Warning, TEXT("Failed to write OSCMessagePacket. Invalid OSCAddress '%s'"), *Address.Value);
+		return;
+	}
+
+	// Begin writing data 
+	Stream.WriteString(Address.Value);
+
+	// write type tags
+	FString TagTypes = ",";
+	for (int32 i = 0; i < Arguments.Num(); i++)
+	{
+		TagTypes += static_cast<ANSICHAR>(Arguments[i].GetTypeTag());
+	}
+
+	Stream.WriteString(TagTypes);
+
+	for (FOSCType OSCType : Arguments)
+	{
+		switch (OSCType.GetTypeTag())
+		{
+		case EOSCTypeTag::CHAR:
+			Stream.WriteChar(OSCType.GetChar());
+			break;
+		case EOSCTypeTag::INT32:
+			Stream.WriteInt32(OSCType.GetInt32());
+			break;
+		case EOSCTypeTag::FLOAT:
+			Stream.WriteFloat(OSCType.GetFloat());
+			break;
+		case EOSCTypeTag::DOUBLE:
+			Stream.WriteDouble(OSCType.GetDouble());
+			break;
+		case EOSCTypeTag::INT64:
+			Stream.WriteInt64(OSCType.GetInt64());
+			break;
+		case EOSCTypeTag::TIME:
+			Stream.WriteUInt64(OSCType.GetTimeTag());
+			break;
+		case EOSCTypeTag::STRING:
+			Stream.WriteString(OSCType.GetString());
+			break;
+		case EOSCTypeTag::BLOB:
+		{
+			TArray<uint8> blob = OSCType.GetBlob();
+			Stream.WriteBlob(blob);
+		}
+		break;
+		case EOSCTypeTag::COLOR:
+#if PLATFORM_LITTLE_ENDIAN
+			Stream.WriteInt32(OSCType.GetColor().ToPackedABGR());
+#else
+			Stream.WriteInt32(OSCType.GetColor().ToPackedRGBA());
+#endif
+			break;
+		case EOSCTypeTag::TRUE:
+		case EOSCTypeTag::FALSE:
+		case EOSCTypeTag::NIL:
+		case EOSCTypeTag::INFINITUM:
+			// No values are written for these types
+			break;
+		default:
+			// Argument is not supported 
+			unimplemented();
+			break;
+		}
+	}
+}
+
+void FOSCMessagePacket::ReadData(FOSCStream& Stream)
+{
+	// Read Address
+	Address = Stream.ReadString();
+
+	// Read string of tags
+	const FString StreamString = Stream.ReadString();
+	const TArray<TCHAR>& TagTypes = StreamString.GetCharArray();
+
+	// Skip the first argument which is ','
+	for (int32 i = 1; i < TagTypes.Num(); i++)
+	{
+		const EOSCTypeTag Tag = static_cast<EOSCTypeTag>(TagTypes[i]);
+		switch (Tag)
+		{
+		case EOSCTypeTag::CHAR:
+			Arguments.Add(FOSCType(Stream.ReadChar()));
+			break;
+		case EOSCTypeTag::INT32:
+			Arguments.Add(FOSCType(Stream.ReadInt32()));
+			break;
+		case EOSCTypeTag::FLOAT:
+			Arguments.Add(FOSCType(Stream.ReadFloat()));
+			break;
+		case EOSCTypeTag::DOUBLE:
+			Arguments.Add(FOSCType(Stream.ReadDouble()));
+			break;
+		case EOSCTypeTag::INT64:
+			Arguments.Add(FOSCType(Stream.ReadInt64()));
+			break;
+		case EOSCTypeTag::TRUE:
+			Arguments.Add(FOSCType(true));
+			break;
+		case EOSCTypeTag::FALSE:
+			Arguments.Add(FOSCType(false));
+			break;
+		case EOSCTypeTag::NIL:
+			Arguments.Add(FOSCType(EOSCTypeTag::NIL));
+			break;
+		case EOSCTypeTag::INFINITUM:
+			Arguments.Add(FOSCType(EOSCTypeTag::INFINITUM));
+			break;
+		case EOSCTypeTag::TIME:
+			Arguments.Add(FOSCType(Stream.ReadUInt64()));
+			break;
+		case EOSCTypeTag::STRING:
+			Arguments.Add(FOSCType(Stream.ReadString()));
+			break;
+		case EOSCTypeTag::BLOB:
+			Arguments.Add(FOSCType(Stream.ReadBlob()));
+			break;
+		case EOSCTypeTag::COLOR:
+			Arguments.Add(FOSCType(FColor(Stream.ReadInt32())));
+			break;
+		case EOSCTypeTag::TERMINATE:
+			Stream.ReadChar();
+			break;
+
+		default:
+			// Argument is not supported 
+			unimplemented();
+			break;
+		}
+	}
+}
diff --git a/Engine/Plugins/Runtime/OSC/Source/OSC/Private/OSCMessagePacket.h b/Engine/Plugins/Runtime/OSC/Source/OSC/Private/OSCMessagePacket.h
new file mode 100644
--- /dev/null
+++ b/Engine/Plugins/Runtime/OSC/Source/OSC/Private/OSCMessagePacket.h
@@ -0,0 +1,43 @@
+// Copyright 1998-2019 Epic Games, Inc. All Rights Reserved.
+#pragma once
+
+#include "OSCAddress.h"
+#include "OSCPacket.h"
+#include "OSCStream.h"
+#include "OSCTypes.h"
+
+
+class FOSCMessagePacket : public IOSCPacket
+{
+public:
+	FOSCMessagePacket();
+	virtual ~FOSCMessagePacket();
+
+	/** Set OSC message address. */
+	void SetAddress(const FOSCAddress& InAddress);
+
+	/** Get OSC message address. */
+	virtual const FOSCAddress& GetAddress() const override;
+
+	/** Get arguments array. */
+	TArray<FOSCType>& GetArguments();
+
+	/** Returns false to indicate type is not OSC bundle. */
+	virtual bool IsBundle();
+
+	/** Returns true to indicate its an OSC message. */
+	virtual bool IsMessage();
+
+	/** Write message data into an OSC stream. */
+	virtual void WriteData(FOSCStream& Stream) override;
+
+	/** Reads message data from an OSC stream and creates new argument. */
+	virtual void ReadData(FOSCStream& Stream) override;
+
+private:
+	/** OSC address. */
+	FOSCAddress Address;
+
+	/** List of argument types. */
+	TArray<FOSCType> Arguments;
+};
diff --git a/Engine/Plugins/Runtime/OSC/Source/OSC/Private/OSCPacket.cpp b/Engine/Plugins/Runtime/OSC/Source/OSC/Private/OSCPacket.cpp
--- a/Engine/Plugins/Runtime/OSC/Source/OSC/Private/OSCPacket.cpp
+++ b/Engine/Plugins/Runtime/OSC/Source/OSC/Private/OSCPacket.cpp
@@ -1,20 +1,20 @@
 // Copyright 1998-2019 Epic Games, Inc. All Rights Reserved.
 #include "OSCPacket.h"
 
-#include "OSCMessage.h"
-#include "OSCBundle.h"
+#include "OSCMessagePacket.h"
+#include "OSCBundlePacket.h"
 #include "OSCLog.h"
 
 
-FOSCPacket::FOSCPacket()
+IOSCPacket::IOSCPacket()
 {
 }
 
-FOSCPacket::~FOSCPacket()
+IOSCPacket::~IOSCPacket()
 {
 }
 
-TSharedPtr<FOSCPacket> FOSCPacket::CreatePacket(const uint8* PacketType)
+TSharedPtr<IOSCPacket> IOSCPacket::CreatePacket(const uint8* PacketType)
 {
 	const FOSCAddress Address(ANSI_TO_TCHAR((const ANSICHAR*)&PacketType[0]));
 	if (Address.IsMessage())
@@ -32,169 +32,3 @@ TSharedPtr<FOSCPacket> FOSCPacket::CreatePacket(const uint8* PacketType)
 		return nullptr;
 	}
 }
-
-FOSCMessagePacket::FOSCMessagePacket()
-	: FOSCPacket()
-{}
-
-FOSCMessagePacket::~FOSCMessagePacket()
-{}
-
-void FOSCMessagePacket::SetAddress(const FOSCAddress& InAddress)
-{
-	Address = InAddress;
-}
-
-const FOSCAddress& FOSCMessagePacket::GetAddress() const
-{
-	return Address;
-}
-
-TArray<FOSCType>& FOSCMessagePacket::GetArguments()
-{
-	return Arguments;
-}
-
-bool FOSCMessagePacket::IsBundle()
-{
-	return false;
-}
-
-bool FOSCMessagePacket::IsMessage()
-{
-	return true;
-}
-
-void FOSCMessagePacket::WriteData(FOSCStream& Stream)
-{
-	if (!Address.IsMessage())
-	{
-		UE_LOG(LogOSC, Warning, TEXT("Failed to write OSCMessagePacket. Invalid OSCAddress '%s'"), *Address.Value);
-		return;
-	}
-
-	// Begin writing data 
-	Stream.WriteString(Address.Value);
-
-	// write type tags
-	FString TagTypes = ",";
-	for (int32 i = 0; i < Arguments.Num(); i++)
-	{
-		TagTypes += static_cast<ANSICHAR>(Arguments[i].GetTypeTag());
-	}
-
-	Stream.WriteString(TagTypes);
-
-	for (FOSCType OSCType : Arguments)
-	{
-		switch (OSCType.GetTypeTag())
-		{
-		case EOSCTypeTag::CHAR:
-			Stream.WriteChar(OSCType.GetChar());
-			break;
-		case EOSCTypeTag::INT32:
-			Stream.WriteInt32(OSCType.GetInt32());
-			break;
-		case EOSCTypeTag::FLOAT:
-			Stream.WriteFloat(OSCType.GetFloat());
-			break;
-		case EOSCTypeTag::DOUBLE:
-			Stream.WriteDouble(OSCType.GetDouble());
-			break;
-		case EOSCTypeTag::INT64:
-			Stream.WriteInt64(OSCType.GetInt64());
-			break;
-		case EOSCTypeTag::TIME:
-			Stream.WriteUInt64(OSCType.GetTimeTag());
-			break;
-		case EOSCTypeTag::STRING:
-			Stream.WriteString(OSCType.GetString());
-			break;
-		case EOSCTypeTag::BLOB:
-		{
-			TArray<uint8> blob = OSCType.GetBlob();
-			Stream.WriteBlob(blob);
-		}
-		break;
-		case EOSCTypeTag::COLOR:
-			Stream.WriteInt32(OSCType.GetColor().GetInt32());
-			break;
-		case EOSCTypeTag::TRUE:
-		case EOSCTypeTag::FALSE:
-		case EOSCTypeTag::NIL:
-		case EOSCTypeTag::INFINITUM:
-			// No values are written for these types
-			break;
-		default:
-			// Argument is not supported 
-			unimplemented();
-			break;
-		}
-	}
-}
-
-void FOSCMessagePacket::ReadData(FOSCStream& Stream)
-{
-	// Read Address
-	Address = Stream.ReadString();
-
-	// Read string of tags
-	const FString StreamString = Stream.ReadString();
-	const TArray<TCHAR>& TagTypes = StreamString.GetCharArray();
-
-	// Skip the first argument which is ','
-	for (int32 i = 1; i < TagTypes.Num(); i++)
-	{
-		const EOSCTypeTag Tag = static_cast<EOSCTypeTag>(TagTypes[i]);
-		switch (Tag)
-		{
-		case EOSCTypeTag::CHAR:
-			Arguments.Add(FOSCType(Stream.ReadChar()));
-			break;
-		case EOSCTypeTag::INT32:
-			Arguments.Add(FOSCType(Stream.ReadInt32()));
-			break;
-		case EOSCTypeTag::FLOAT:
-			Arguments.Add(FOSCType(Stream.ReadFloat()));
-			break;
-		case EOSCTypeTag::DOUBLE:
-			Arguments.Add(FOSCType(Stream.ReadDouble()));
-			break;
-		case EOSCTypeTag::INT64:
-			Arguments.Add(FOSCType(Stream.ReadInt64()));
-			break;
-		case EOSCTypeTag::TRUE:
-			Arguments.Add(FOSCType(true));
-			break;
-		case EOSCTypeTag::FALSE:
-			Arguments.Add(FOSCType(false));
-			break;
-		case EOSCTypeTag::NIL:
-			Arguments.Add(FOSCType(EOSCTypeTag::NIL));
-			break;
-		case EOSCTypeTag::INFINITUM:
-			Arguments.Add(FOSCType(EOSCTypeTag::INFINITUM));
-			break;
-		case EOSCTypeTag::TIME:
-			Arguments.Add(FOSCType(Stream.ReadUInt64()));
-			break;
-		case EOSCTypeTag::STRING:
-			Arguments.Add(FOSCType(Stream.ReadString()));
-			break;
-		case EOSCTypeTag::BLOB:
-			Arguments.Add(FOSCType(Stream.ReadBlob()));
-			break;
-		case EOSCTypeTag::COLOR:
-			Arguments.Add(FOSCType(FOSCColor(Stream.ReadInt32())));
-			break;
-		case EOSCTypeTag::TERMINATE:
-			Stream.ReadChar();
-			break;
-
-		default:
-			// Argument is not supported 
-			unimplemented();
-			break;
-		}
-	}
-}
diff --git a/Engine/Plugins/Runtime/OSC/Source/OSC/Private/OSCPacket.h b/Engine/Plugins/Runtime/OSC/Source/OSC/Private/OSCPacket.h
deleted file mode 100644
diff --git a/Engine/Plugins/Runtime/OSC/Source/OSC/Private/OSCServer.cpp b/Engine/Plugins/Runtime/OSC/Source/OSC/Private/OSCServer.cpp
--- a/Engine/Plugins/Runtime/OSC/Source/OSC/Private/OSCServer.cpp
+++ b/Engine/Plugins/Runtime/OSC/Source/OSC/Private/OSCServer.cpp
@@ -8,18 +8,23 @@
 
 #include "OSCStream.h"
 #include "OSCMessage.h"
+#include "OSCMessagePacket.h"
 #include "OSCBundle.h"
+#include "OSCBundlePacket.h"
 #include "OSCLog.h"
 
-#define OSC_BUFFER_FREE 0
-#define OSC_BUFFER_LOCKED 1
+#define OSC_SERVER_BUFFER_FREE 0
+#define OSC_SERVER_BUFFER_LOCKED 1
 
 
 UOSCServer::UOSCServer(const FObjectInitializer& ObjectInitializer)
 	: Super(ObjectInitializer) 
 	, bWhitelistClients(false)
-	, OSCPackets(1024)
+	, Socket(nullptr)
+	, SocketReceiver(nullptr)
 	, OSCBufferLock(0)
+	, Port(0)
+	, bMulticastLoopback(false)
 {
 }
 
@@ -31,47 +36,46 @@ void UOSCServer::Callback(const FArrayReaderPtr& Data, const FIPv4Endpoint& Endp
 		return;
 	}
 
-	TSharedPtr<FOSCPacket> Packet = FOSCPacket::CreatePacket(Data->GetData());
+	TSharedPtr<IOSCPacket> Packet = IOSCPacket::CreatePacket(Data->GetData());
 	if (Packet.IsValid())
 	{
 		FOSCStream Stream = FOSCStream(Data->GetData(), Data->Num());
 
 		Packet->ReadData(Stream);
-		OSCPackets.Buffer.Enqueue(MoveTemp(Packet));
+		OSCPackets.Enqueue(MoveTemp(Packet));
 	}
 
-	if (!FPlatformAtomics::InterlockedCompareExchange(&OSCBufferLock, OSC_BUFFER_LOCKED, OSC_BUFFER_FREE) &&
-		!OSCPackets.Buffer.IsEmpty())
+	if (!FPlatformAtomics::InterlockedCompareExchange(&OSCBufferLock, OSC_SERVER_BUFFER_LOCKED, OSC_SERVER_BUFFER_FREE) &&
+		!OSCPackets.IsEmpty())
 	{
-		check(OSCBufferLock == OSC_BUFFER_LOCKED);
-		FSimpleDelegateGraphTask::CreateAndDispatchWhenReady(
-			FSimpleDelegateGraphTask::FDelegate::CreateLambda([this]() {
-			check(OSCBufferLock == OSC_BUFFER_LOCKED);
-			FPlatformAtomics::InterlockedCompareExchange(&OSCBufferLock, OSC_BUFFER_FREE, OSC_BUFFER_LOCKED);
+		check(OSCBufferLock == OSC_SERVER_BUFFER_LOCKED);
+		FSimpleDelegateGraphTask::CreateAndDispatchWhenReady(FSimpleDelegateGraphTask::FDelegate::CreateLambda([this]()
+		{
+			check(OSCBufferLock == OSC_SERVER_BUFFER_LOCKED);
+			FPlatformAtomics::InterlockedCompareExchange(&OSCBufferLock, OSC_SERVER_BUFFER_FREE, OSC_SERVER_BUFFER_LOCKED);
 				
-				TSharedPtr<FOSCPacket> Packet;
-				while (OSCPackets.Buffer.Dequeue(Packet))
+			TSharedPtr<IOSCPacket> Packet;
+			while (OSCPackets.Dequeue(Packet))
+			{
+				if (Packet->IsMessage())
+				{
+					FOSCMessage Message(FOSCMessage(StaticCastSharedPtr<FOSCMessagePacket>(Packet)));
+					OnOscReceived.Broadcast(Message);
+				}
+				else if (Packet->IsBundle())
+				{
+					FOSCBundle Bundle(StaticCastSharedPtr<FOSCBundlePacket>(Packet));
+					OnOscBundleReceived.Broadcast(Bundle);
+				}
+				else
 				{
-					if (Packet->IsMessage())
-					{
-						FOSCMessage Message(FOSCMessage(StaticCastSharedPtr<FOSCMessagePacket>(Packet)));
-						OnOscReceived.Broadcast(Message);
-					}
-					else if (Packet->IsBundle())
-					{
-						FOSCBundle Bundle(StaticCastSharedPtr<FOSCBundlePacket>(Packet));
-						OnOscBundleReceived.Broadcast(Bundle);
-					}
-					else
-					{
-						UE_LOG(LogOSC, Warning, TEXT("Failed to parse invalid received OSC message. OSCAddress '%s' is invalid."), *Packet->GetAddress().Value);
-					}
+					UE_LOG(LogOSC, Warning, TEXT("Failed to parse invalid received OSC message. OSCAddress '%s' is invalid."), *Packet->GetAddress().Value);
 				}
-				OSCPackets.Buffer.Empty();
-			}),
-			TStatId(),
-			nullptr,
-			ENamedThreads::GameThread);
+			}
+		}),
+		TStatId(),
+		nullptr,
+		ENamedThreads::GameThread);
 	}
 }
 
diff --git a/Engine/Plugins/Runtime/OSC/Source/OSC/Private/OSCStream.cpp b/Engine/Plugins/Runtime/OSC/Source/OSC/Private/OSCStream.cpp
--- a/Engine/Plugins/Runtime/OSC/Source/OSC/Private/OSCStream.cpp
+++ b/Engine/Plugins/Runtime/OSC/Source/OSC/Private/OSCStream.cpp
@@ -64,6 +64,23 @@ void FOSCStream::WriteChar(TCHAR Char)
 	Write(ToWrite, 1);
 }
 
+FColor FOSCStream::ReadColor()
+{
+	uint32 Packed = static_cast<uint32>(ReadInt32());
+	return FColor(Packed);
+}
+
+void FOSCStream::WriteColor(FColor Color)
+{
+#if PLATFORM_LITTLE_ENDIAN
+	uint32 Packed = Color.ToPackedABGR();
+#else // PLATFORM_LITTLE_ENDIAN
+	uint32 Packed = Color.ToPackedRGBA();
+#endif // !PLATFORM_LITTLE_ENDIAN
+
+	WriteInt32(static_cast<int32>(Packed));
+}
+
 int32 FOSCStream::ReadInt32()
 {
 	uint8 Temp[4];
@@ -83,7 +100,7 @@ int32 FOSCStream::ReadInt32()
 		return u.i;
 #else
 		return *(int32*)(Temp);
-#endif
+#endif // !PLATFORM_LITTLE_ENDIAN
 	}
 
 	return 0;
@@ -105,9 +122,9 @@ void FOSCStream::WriteInt32(int32 Value)
 	Temp[2] = u.c[1];
 	Temp[1] = u.c[2];
 	Temp[0] = u.c[3];
-#else
-	*reinterpret_cast<int32*>(Temp) = Value;
-#endif
+#else // PLATFORM_LITTLE_ENDIAN
+	*(int32*)(Temp) = Value;
+#endif // !PLATFORM_LITTLE_ENDIAN
 	Write(Temp, 4);
 }
 
@@ -132,9 +149,9 @@ double FOSCStream::ReadDouble()
 		u.c[7] = Temp[0];
 
 		return u.d;
-#else
+#else // PLATFORM_LITTLE_ENDIAN
 		return *(double*)Temp;
-#endif
+#endif // !PLATFORM_LITTLE_ENDIAN
 
 	}
 
@@ -161,9 +178,9 @@ void FOSCStream::WriteDouble(uint64 Value)
 	Temp[2] = u.c[5];
 	Temp[1] = u.c[6];
 	Temp[0] = u.c[7];
-#else
-	*reinterpret_cast<double*>(Temp) = Value;
-#endif
+#else // PLATFORM_LITTLE_ENDIAN
+	*(double*)(Temp) = Value;
+#endif // !PLATFORM_LITTLE_ENDIAN
 	Write(Temp, 8);
 }
 
@@ -188,9 +205,9 @@ int64 FOSCStream::ReadInt64()
 		u.c[7] = Temp[0];
 
 		return u.i;
-#else
+#else // PLATFORM_LITTLE_ENDIAN
 		return *(int64*)Temp;
-#endif
+#endif // !PLATFORM_LITTLE_ENDIAN
 	}
 
 	return 0;
@@ -216,9 +233,9 @@ void FOSCStream::WriteInt64(int64 Value)
 	Temp[2] = u.c[5];
 	Temp[1] = u.c[6];
 	Temp[0] = u.c[7];
-#else
-	*reinterpret_cast<int64*>(Temp) = Value;
-#endif
+#else // PLATFORM_LITTLE_ENDIAN
+	*(int64*)(Temp) = Value;
+#endif // !PLATFORM_LITTLE_ENDIAN
 	Write(Temp, 8);
 }
 
@@ -243,9 +260,9 @@ uint64 FOSCStream::ReadUInt64()
 		u.c[7] = Temp[0];
 
 		return u.i;
-#else
+#else // PLATFORM_LITTLE_ENDIAN
 		return *(uint64*)Temp;
-#endif
+#endif // !PLATFORM_LITTLE_ENDIAN
 
 	}
 
@@ -272,9 +289,9 @@ void FOSCStream::WriteUInt64(uint64 Value)
 	Temp[2] = u.c[5];
 	Temp[1] = u.c[6];
 	Temp[0] = u.c[7];
-#else
+#else // PLATFORM_LITTLE_ENDIAN
 	*reinterpret_cast<uint64*>(Temp) = Value;
-#endif
+#endif // !PLATFORM_LITTLE_ENDIAN
 	Write(Temp, 8);
 }
 
@@ -295,9 +312,9 @@ float FOSCStream::ReadFloat()
 		u.c[3] = Temp[0];
 
 		return u.f;
-#else
+#else // PLATFORM_LITTLE_ENDIAN
 		return *(float*)Temp;
-#endif
+#endif // !PLATFORM_LITTLE_ENDIAN
 	}
 
 	return 0.0f;
@@ -320,7 +337,7 @@ void FOSCStream::WriteFloat(float Value)
 	Temp[1] = u.c[2];
 	Temp[0] = u.c[3];
 #else // if !PLATFORM_LITTLE_ENDIAN
-	*reinterpret_cast<float*>(Temp) = Value;
+	*(float*)(Temp) = Value;
 #endif // !PLATFORM_LITTLE_ENDIAN
 
 	Write(Temp, 4);
diff --git a/Engine/Plugins/Runtime/OSC/Source/OSC/Public/OSCBundle.h b/Engine/Plugins/Runtime/OSC/Source/OSC/Public/OSCBundle.h
--- a/Engine/Plugins/Runtime/OSC/Source/OSC/Public/OSCBundle.h
+++ b/Engine/Plugins/Runtime/OSC/Source/OSC/Public/OSCBundle.h
@@ -10,31 +10,13 @@ struct OSC_API FOSCBundle
 {
 	GENERATED_USTRUCT_BODY()
 
-	FOSCBundle()
-		: Packet(MakeShareable(new FOSCBundlePacket()))
-	{
-	}
+	FOSCBundle();
+	FOSCBundle(const TSharedPtr<IOSCPacket>& InPacket);
+	~FOSCBundle();
 
-	FOSCBundle(const TSharedPtr<FOSCBundlePacket>& InPacket)
-		: Packet(InPacket)
-	{
-	}
-
-	~FOSCBundle()
-	{
-		Packet.Reset();
-	}
-
-	void SetPacket(const TSharedPtr<FOSCBundlePacket>& InPacket)
-	{
-		Packet = InPacket;
-	}
-
-	const TSharedPtr<FOSCBundlePacket>& GetPacket() const
-	{
-		return Packet;
-	}
+	void SetPacket(const TSharedPtr<IOSCPacket>& InPacket);
+	const TSharedPtr<IOSCPacket>& GetPacket() const;
 
 private:
-	TSharedPtr<FOSCBundlePacket> Packet;
+	TSharedPtr<IOSCPacket> Packet;
 };
\ No newline at end of file
diff --git a/Engine/Plugins/Runtime/OSC/Source/OSC/Public/OSCClient.h b/Engine/Plugins/Runtime/OSC/Source/OSC/Public/OSCClient.h
--- a/Engine/Plugins/Runtime/OSC/Source/OSC/Public/OSCClient.h
+++ b/Engine/Plugins/Runtime/OSC/Source/OSC/Public/OSCClient.h
@@ -45,7 +45,7 @@ protected:
 	void Stop();
 	
 	/** Send OSC packet data. */
-	void SendPacket(FOSCPacket& InPacket);
+	void SendPacket(IOSCPacket& InPacket);
 
 	/** Socket used to send the OSC packets. */
 	FSocket* Socket;
diff --git a/Engine/Plugins/Runtime/OSC/Source/OSC/Public/OSCManager.cpp b/Engine/Plugins/Runtime/OSC/Source/OSC/Public/OSCManager.cpp
deleted file mode 100644
diff --git a/Engine/Plugins/Runtime/OSC/Source/OSC/Public/OSCMessage.h b/Engine/Plugins/Runtime/OSC/Source/OSC/Public/OSCMessage.h
--- a/Engine/Plugins/Runtime/OSC/Source/OSC/Public/OSCMessage.h
+++ b/Engine/Plugins/Runtime/OSC/Source/OSC/Public/OSCMessage.h
@@ -4,15 +4,13 @@
 #include "CoreMinimal.h"
 #include "OSCLog.h"
 #include "OSCPacket.h"
+#include "OSCStream.h"
 #include "UObject/ObjectMacros.h"
 #include "UObject/Object.h"
 #include "UObject/WeakObjectPtr.h"
 
 #include "OSCMessage.generated.h"
 
-// Forward Declarations
-class FOSCStream;
-
 
 USTRUCT(BlueprintType)
 struct OSC_API FOSCMessage
@@ -20,15 +18,15 @@ struct OSC_API FOSCMessage
 	GENERATED_USTRUCT_BODY()
 
 	FOSCMessage();
-	FOSCMessage(const TSharedPtr<FOSCMessagePacket>& InPacket);
+	FOSCMessage(const TSharedPtr<IOSCPacket>& InPacket);
 	~FOSCMessage();
 
-	void SetPacket(TSharedPtr<FOSCMessagePacket>& InPacket);
-	const TSharedPtr<FOSCMessagePacket>& GetPacket() const;
+	void SetPacket(TSharedPtr<IOSCPacket>& InPacket);
+	const TSharedPtr<IOSCPacket>& GetPacket() const;
 
 	bool SetAddress(const FOSCAddress& InAddress);
 	const FOSCAddress& GetAddress() const;
 
 private:
-	TSharedPtr<FOSCMessagePacket> Packet;
+	TSharedPtr<IOSCPacket> Packet;
 };
\ No newline at end of file
diff --git a/Engine/Plugins/Runtime/OSC/Source/OSC/Public/OSCPacket.h b/Engine/Plugins/Runtime/OSC/Source/OSC/Public/OSCPacket.h
new file mode 100644
--- /dev/null
+++ b/Engine/Plugins/Runtime/OSC/Source/OSC/Public/OSCPacket.h
@@ -0,0 +1,38 @@
+// Copyright 1998-2019 Epic Games, Inc. All Rights Reserved.
+#pragma once
+
+#include "OSCTypes.h"
+#include "OSCStream.h"
+#include "OSCAddress.h"
+
+// Forward Declarations
+class FOSCStream;
+
+
+class IOSCPacket
+{
+public:
+	IOSCPacket();
+	virtual ~IOSCPacket();
+
+	/** Write packet data into stream */
+	virtual void WriteData(FOSCStream& Stream) = 0;
+	
+	/** Read packet data from stream */
+	virtual void ReadData(FOSCStream& Stream) = 0;
+	
+	/** Returns OSC address identifier of packet */
+	virtual const FOSCAddress& GetAddress() const = 0;
+
+	/** Set the address identifier of packet */
+	virtual void SetAddress(const FOSCAddress& Address) = 0;
+
+	/** Returns true if packet is message */
+	virtual bool IsMessage() = 0;
+	
+	/** Returns true if packet is bundle */
+	virtual bool IsBundle() = 0;
+
+	/** Create an OSC packet according to the input data. */
+	static TSharedPtr<IOSCPacket> CreatePacket(const uint8* PacketType);
+};
diff --git a/Engine/Plugins/Runtime/OSC/Source/OSC/Public/OSCServer.h b/Engine/Plugins/Runtime/OSC/Source/OSC/Public/OSCServer.h
--- a/Engine/Plugins/Runtime/OSC/Source/OSC/Public/OSCServer.h
+++ b/Engine/Plugins/Runtime/OSC/Source/OSC/Public/OSCServer.h
@@ -2,13 +2,14 @@
 #pragma once
 
 #include "CoreMinimal.h"
-#include "Containers/CircularQueue.h"
+
 #include "UObject/Object.h"
 #include "Interfaces/IPv4/IPv4Address.h"
 #include "Interfaces/IPv4/IPv4Endpoint.h"
 #include "Common/UdpSocketReceiver.h"
 #include "OSCMessage.h"
 #include "OSCBundle.h"
+#include "OSCPacket.h"
 #include "OSCServer.generated.h"
 
 DECLARE_DYNAMIC_MULTICAST_DELEGATE_OneParam(FOSCReceivedMessageEvent, const FOSCMessage &, Message);
@@ -21,21 +22,6 @@ class OSC_API UOSCServer : public UObject
 {
 	GENERATED_UCLASS_BODY()
 
-	struct OscCircularBuffer
-	{
-	public:
-
-		OscCircularBuffer(uint32 CircularQueueSize)
-			: Buffer(CircularQueueSize)
-		{}
-
-		OscCircularBuffer()
-			: Buffer(1)
-		{ }
-
-		TCircularQueue<TSharedPtr<FOSCPacket>> Buffer;
-	};
-
 public:
 	/** Gets whether or not to loopback if ReceiveIPAddress provided is multicast. */
 	UFUNCTION(BlueprintCallable, Category = "Audio|OSC")
@@ -73,28 +59,31 @@ public:
 	UPROPERTY(BlueprintReadWrite, Category = "Audio|OSC")
 	bool bWhitelistClients;
 
+	/** Adds client to whitelist of clients to listen for. */
 	UFUNCTION(BlueprintCallable, Category = "Audio|OSC")
 	void AddWhitelistedClient(const FString& IPAddress);
 
+	/** Removes whitelisted client to listen for. */
 	UFUNCTION(BlueprintCallable, Category = "Audio|OSC")
 	void RemoveWhitelistedClient(const FString& IPAddress);
 
+	/** Clears client whitelist to listen for. */
 	UFUNCTION(BlueprintCallable, Category = "Audio|OSC")
 	void ClearWhitelistedClients();
 
+	/** Returns set of whitelisted clients. */
 	UFUNCTION(BlueprintCallable, Category = "Audio|OSC")
 	TSet<FString> GetWhitelistedClients() const;
 
 protected:
-
 	void BeginDestroy() override;
 	
 	/** Callback that receives data from a socket. */
 	void Callback(const FArrayReaderPtr& Data, const FIPv4Endpoint& Endpoint);
 	
 	/** Circular buffer that stores OSC packets. */
-	OscCircularBuffer OSCPackets;
-	
+	TQueue<TSharedPtr<IOSCPacket>> OSCPackets;
+
 	/** Socket used to listen for OSC packets. */
 	FSocket* Socket;
 
diff --git a/Engine/Plugins/Runtime/OSC/Source/OSC/Public/OSCStream.h b/Engine/Plugins/Runtime/OSC/Source/OSC/Public/OSCStream.h
--- a/Engine/Plugins/Runtime/OSC/Source/OSC/Public/OSCStream.h
+++ b/Engine/Plugins/Runtime/OSC/Source/OSC/Public/OSCStream.h
@@ -31,12 +31,18 @@ public:
 	/** Write an ansi char into the stream */
 	void WriteChar(TCHAR Char);
 
+	/** Write Color into the stream */
+	void WriteColor(FColor Color);
+
+	/** Read Color from the stream */
+	FColor ReadColor();
+
 	/** Read Int32 from the stream */
 	int32 ReadInt32();
 
 	/** Write Int32 into the stream */
 	void WriteInt32(int32 Value);
-	
+
 	/** Read Double from the stream */
 	double ReadDouble();
 
diff --git a/Engine/Plugins/Runtime/OSC/Source/OSC/Public/OSCTypes.h b/Engine/Plugins/Runtime/OSC/Source/OSC/Public/OSCTypes.h
--- a/Engine/Plugins/Runtime/OSC/Source/OSC/Public/OSCTypes.h
+++ b/Engine/Plugins/Runtime/OSC/Source/OSC/Public/OSCTypes.h
@@ -3,6 +3,9 @@
 
 #include "CoreMinimal.h"
 
+#include "Math/Color.h"
+
+
 enum class EOSCTypeTag {
 
 	INT32 = 'i',
@@ -21,31 +24,6 @@ enum class EOSCTypeTag {
 	TERMINATE = '\0'
 };
 
-struct FOSCColor
-{
-	explicit FOSCColor() 
-	{}
-
-	explicit FOSCColor(uint8 R, uint8 G, uint8 B, uint8 A) 
-		: Red(R), Green(G), Blue(B), Alpha(A)
-	{}
-
-	explicit FOSCColor(int32 c)
-	{
-		Red = (uint8)(c >> 24);
-		Green = (uint8)(c >> 16);
-		Blue = (uint8)(c >> 8),
-		Alpha = (uint8)c;
-	}
-
-	uint8 Red, Green, Blue, Alpha;
-
-	int32 GetInt32() const
-	{
-		return	(static_cast<int32>(Red)) | (static_cast<int32>(Green) << 8) | 
-				(static_cast<int32>(Blue) << 16) | (static_cast<int32>(Alpha) << 24);
-	}
-};
 
 class FOSCType
 {
@@ -93,29 +71,29 @@ public:
 		, Blob()
 	{
 	}
-	explicit FOSCType(FString Value)
+	explicit FOSCType(const FString& Value)
 		: TypeTag(EOSCTypeTag::STRING)
 		, Data(0)
 		, String(Value)
 		, Blob()
 	{
 	}
-	explicit FOSCType(TArray<uint8> Value)
+	explicit FOSCType(const TArray<uint8>& Value)
 		: TypeTag(EOSCTypeTag::BLOB)
 		, Data(0)
 		, Blob(Value)
 	{
 	}
-	explicit FOSCType(FOSCColor color)
+	explicit FOSCType(FColor Value)
 		: TypeTag(EOSCTypeTag::COLOR)
 		, Data(0)
 		, Blob()
-		, Color(color)
+		, Color(Value)
 	{
 	}
 
-	explicit FOSCType(EOSCTypeTag typeTag)
-		: TypeTag(typeTag)
+	explicit FOSCType(EOSCTypeTag TypeTag)
+		: TypeTag(TypeTag)
 		, Data(0)
 		, Blob()
 	{
@@ -151,7 +129,7 @@ public:
 	TArray<uint8> GetBlob() const { return Blob; }
 
 	bool IsColor() const { return TypeTag == EOSCTypeTag::COLOR; }
-	FOSCColor GetColor()  const { return Color; }
+	FColor GetColor()  const { return Color; }
 
 	bool IsNil() const { return TypeTag == EOSCTypeTag::NIL; }
 	bool IsInfinitum() const { return TypeTag == EOSCTypeTag::INFINITUM; }
@@ -199,5 +177,5 @@ private:
 	DataTypes Data;
 	FString String;
 	TArray<uint8> Blob;
-	FOSCColor Color;
+	FColor Color;
 };
