Fix non-unity build

diff --git a/Engine/Plugins/Runtime/OSC/Source/OSC/Public/OSCServer.h b/Engine/Plugins/Runtime/OSC/Source/OSC/Public/OSCServer.h
--- a/Engine/Plugins/Runtime/OSC/Source/OSC/Public/OSCServer.h
+++ b/Engine/Plugins/Runtime/OSC/Source/OSC/Public/OSCServer.h
@@ -3,15 +3,18 @@
 
 #include "CoreMinimal.h"
 
-#include "UObject/Object.h"
+#include "Common/UdpSocketReceiver.h"
+#include "Containers/Queue.h"
 #include "Interfaces/IPv4/IPv4Address.h"
 #include "Interfaces/IPv4/IPv4Endpoint.h"
-#include "Common/UdpSocketReceiver.h"
 #include "OSCMessage.h"
 #include "OSCBundle.h"
 #include "OSCPacket.h"
+#include "UObject/Object.h"
+
 #include "OSCServer.generated.h"
 
+
 DECLARE_DYNAMIC_MULTICAST_DELEGATE_OneParam(FOSCReceivedMessageEvent, const FOSCMessage &, Message);
 DECLARE_DYNAMIC_MULTICAST_DELEGATE_OneParam(FOSCReceivedBundleEvent, const FOSCBundle &, Bundle);
 
@@ -81,7 +84,7 @@ protected:
 	/** Callback that receives data from a socket. */
 	void Callback(const FArrayReaderPtr& Data, const FIPv4Endpoint& Endpoint);
 	
-	/** Circular buffer that stores OSC packets. */
+	/** Queue stores incoming OSC packet requests to process on the game thread. */
 	TQueue<TSharedPtr<IOSCPacket>> OSCPackets;
 
 	/** Socket used to listen for OSC packets. */
