- Fix unity builds
- Add OSC Address Matching API

diff --git a/Engine/Plugins/Runtime/OSC/Source/OSC/OSC.Build.cs b/Engine/Plugins/Runtime/OSC/Source/OSC/OSC.Build.cs
--- a/Engine/Plugins/Runtime/OSC/Source/OSC/OSC.Build.cs
+++ b/Engine/Plugins/Runtime/OSC/Source/OSC/OSC.Build.cs
@@ -6,6 +6,8 @@ namespace UnrealBuildTool.Rules
 	{
         public OSC(ReadOnlyTargetRules Target) : base(Target)
 		{
+            OptimizeCode = CodeOptimization.Never;
+
             PublicDependencyModuleNames.AddRange(
 				new string[] {
                     "Core",
diff --git a/Engine/Plugins/Runtime/OSC/Source/OSC/Private/OSCAddress.cpp b/Engine/Plugins/Runtime/OSC/Source/OSC/Private/OSCAddress.cpp
new file mode 100644
--- /dev/null
+++ b/Engine/Plugins/Runtime/OSC/Source/OSC/Private/OSCAddress.cpp
@@ -0,0 +1,191 @@
+// Copyright 1998-2019 Epic Games, Inc. All Rights Reserved.
+#include "OSCAddress.h"
+
+// TODO: Finish implementing full pattern matching instead of just substring parsing
+#define USE_FULL_OSC_PATTERN_MATCH 0
+
+
+FOSCAddress::FOSCAddress()
+	: bIsValidPattern(false)
+	, bIsValidPath(false)
+	, Hash(GetTypeHash(GetFullPath()))
+{
+}
+
+FOSCAddress::FOSCAddress(const FString& InValue)
+	: bIsValidPattern(false)
+	, bIsValidPath(false)
+{
+	bIsValidPath = IsValidPath(InValue, false);
+	InValue.ParseIntoArray(Containers, *OSC::PathSeparator, true);
+
+	if (Containers.Num() > 0)
+	{
+		Method = Containers.Pop();
+	}
+
+	Hash = GetTypeHash(InValue);
+	bIsValidPattern = IsValidPattern(Containers, Method);
+}
+
+// TODO: Pattern should check for validity of search strings.
+// ex. '[' doesn't come before closing ']', etc.  This is essentially
+// here as a stub.
+bool FOSCAddress::IsValidPattern(const TArray<FString>& InContainers, const FString& InMethod)
+{
+#if USE_FULL_OSC_PATTERN_MATCH
+	for (const FString& Container : InContainers)
+	{
+		for (const TCHAR& Char : Container.GetCharArray())
+		{
+			check(Char != OSC::PathSeparator[0]);
+
+			if (OSC::PatternChars.Contains(Char))
+			{
+				return true;
+			}
+		}
+	}
+	return false;
+#else // USE_FULL_OSC_PATTERN_MATCH
+	return true;
+#endif // !USE_FULL_OSC_PATTERN_MATCH
+}
+
+bool FOSCAddress::IsValidPath(const FString& Path, bool bInvalidateSeparator)
+{
+	for (const TCHAR& Char : Path.GetCharArray())
+	{
+		if (!bInvalidateSeparator)
+		{
+			if (Char == OSC::PathSeparator[0])
+			{
+				continue;
+			}
+		}
+
+		if (OSC::InvalidChars.Contains(Char))
+		{
+			return false;
+		}
+
+		if (OSC::PatternChars.Contains(Char))
+		{
+			return false;
+		}
+	}
+
+	return true;
+}
+
+bool FOSCAddress::Matches(const FOSCAddress& Address) const
+{
+	if (IsValidPattern())
+	{
+#if USE_FULL_OSC_PATTERN_MATCH
+		// TODO: Implement
+		return false;
+#else // USE_FULL_OSC_PATTERN_MATCH
+		return Address.GetFullPath().Contains(GetFullPath(), ESearchCase::CaseSensitive);
+#endif // !USE_FULL_OSC_PATTERN_MATCH
+	}
+	else
+	{
+		return IsValidPath() && Address.IsValidPath() && Address.GetFullPath() == GetFullPath();
+	}
+}
+
+uint32 FOSCAddress::GetHash() const
+{
+	return Hash;
+}
+
+bool FOSCAddress::IsValidPattern() const
+{
+	return bIsValidPattern;
+}
+
+bool FOSCAddress::IsValidPath() const
+{
+	return bIsValidPath;
+}
+
+void FOSCAddress::PushContainer(const FString& Container)
+{
+	const bool bInvalidateSeparator = true;
+	if (!IsValidPath(Container, bInvalidateSeparator))
+	{
+		UE_LOG(LogOSC, Warning, TEXT("Failed to push container on OSCAddress. "
+			"Cannot contain invalid or reserved OSC pattern identifiers."));
+		return;
+	}
+
+	Containers.Push(Container);
+	Hash = GetTypeHash(GetFullPath());
+}
+
+void FOSCAddress::ClearContainers(const FString& Container)
+{
+	Containers.Reset();
+	Hash = GetTypeHash(GetFullPath());
+}
+
+void FOSCAddress::SetMethod(const FString& InMethod)
+{
+	const bool bInvalidateSeparator = true;
+	if (!IsValidPath(Method, bInvalidateSeparator))
+	{
+		UE_LOG(LogOSC, Warning, TEXT("Failed to set method on OSCAddress. "
+			"Cannot contain invalid or reserved OSC pattern identifiers."));
+		return;
+	}
+
+	Method = InMethod;
+	Hash = GetTypeHash(GetFullPath());
+}
+
+FString FOSCAddress::PopContainer()
+{
+	if (Containers.Num() > 0)
+	{
+		FString Popped = Containers.Pop(false);
+		Hash = GetTypeHash(GetFullPath());
+	}
+
+	return FString();
+}
+
+FString FOSCAddress::GetContainerPath() const
+{
+	return OSC::PathSeparator + FString::Join(Containers, *OSC::PathSeparator);
+}
+
+FString FOSCAddress::GetContainer(int32 Index) const
+{
+	if (Index < Containers.Num())
+	{
+		return Containers[Index];
+	}
+
+	return FString();
+}
+
+void FOSCAddress::GetContainers(TArray<FString>& OutContainers) const
+{
+	OutContainers = Containers;
+}
+
+FString FOSCAddress::GetFullPath() const
+{
+	if (Containers.Num() == 0)
+	{
+		return OSC::PathSeparator + Method;
+	}
+
+	return GetContainerPath() + OSC::PathSeparator + Method;
+}
+
+const FString& FOSCAddress::GetMethodName() const
+{
+	return Method;
+}
diff --git a/Engine/Plugins/Runtime/OSC/Source/OSC/Private/OSCBundlePacket.cpp b/Engine/Plugins/Runtime/OSC/Source/OSC/Private/OSCBundlePacket.cpp
--- a/Engine/Plugins/Runtime/OSC/Source/OSC/Private/OSCBundlePacket.cpp
+++ b/Engine/Plugins/Runtime/OSC/Source/OSC/Private/OSCBundlePacket.cpp
@@ -29,14 +29,8 @@ uint64 FOSCBundlePacket::GetTimeTag() const
 
 void FOSCBundlePacket::WriteData(FOSCStream& Stream)
 {
-	if (!GetAddress().IsBundle())
-	{
-		UE_LOG(LogOSC, Warning, TEXT("Failed to write OSCBundlePacket. Invalid OSCAddress '%s'"), *GetAddress().Value);
-		return;
-	}
-
 	// Write bundle & time tag
-	Stream.WriteString(GetAddress().Value);
+	Stream.WriteString(OSC::BundleTag);
 	Stream.WriteUInt64(GetTimeTag());
 
 	for (TSharedPtr<IOSCPacket>& Packet : Packets)
diff --git a/Engine/Plugins/Runtime/OSC/Source/OSC/Private/OSCClient.cpp b/Engine/Plugins/Runtime/OSC/Source/OSC/Private/OSCClient.cpp
--- a/Engine/Plugins/Runtime/OSC/Source/OSC/Private/OSCClient.cpp
+++ b/Engine/Plugins/Runtime/OSC/Source/OSC/Private/OSCClient.cpp
@@ -55,9 +55,9 @@ void UOSCClient::GetSendIPAddress(FString& InIPAddress, int32& Port)
 void UOSCClient::SendPacket(IOSCPacket& Packet)
 {
 	const FOSCAddress& OSCAddress = Packet.GetAddress();
-	if (!OSCAddress.IsValid())
+	if (Packet.IsMessage() && !OSCAddress.IsValidPath())
 	{
-		UE_LOG(LogOSC, Warning, TEXT("Failed to write bundle data. Invalid OSCAddress '%s'"), *OSCAddress.Value);
+		UE_LOG(LogOSC, Warning, TEXT("Failed to write packet data. Invalid OSCAddress '%s'"), *OSCAddress.GetFullPath());
 		return;
 	}
 
@@ -77,7 +77,7 @@ void UOSCClient::SendPacket(IOSCPacket& Packet)
 			if (!bSuccess || BytesSent <= 0)
 			{
 				UE_LOG(LogOSC, Verbose, TEXT("OSC Packet failed: Client '%s', OSCAddress '%s', Send IPAddress %s, Attempted Bytes = %d"),
-					*GetName(), *OSCAddress.Value, *IPAddress->ToString(true /*bAppendPort*/), AttemptedLength);
+					*GetName(), *OSCAddress.GetFullPath(), *IPAddress->ToString(true /*bAppendPort*/), AttemptedLength);
 				return;
 			}
 
@@ -86,7 +86,7 @@ void UOSCClient::SendPacket(IOSCPacket& Packet)
 		}
 
 		UE_LOG(LogOSC, Verbose, TEXT("OSC Packet sent: Client '%s', OSCAddress '%s', Send IPAddress %s, Bytes Sent = %d"),
-			*GetName(), *OSCAddress.Value, *IPAddress->ToString(true /*bAppendPort*/), AttemptedLength);
+			*GetName(), *OSCAddress.GetFullPath(), *IPAddress->ToString(true /*bAppendPort*/), AttemptedLength);
 	}
 }
 
diff --git a/Engine/Plugins/Runtime/OSC/Source/OSC/Private/OSCManager.cpp b/Engine/Plugins/Runtime/OSC/Source/OSC/Private/OSCManager.cpp
--- a/Engine/Plugins/Runtime/OSC/Source/OSC/Private/OSCManager.cpp
+++ b/Engine/Plugins/Runtime/OSC/Source/OSC/Private/OSCManager.cpp
@@ -350,19 +350,14 @@ void UOSCManager::GetBlob(const FOSCMessage& Message, const int32 Index, TArray<
 	}
 }
 
-bool UOSCManager::OSCAddressIsBundle(const FOSCAddress& Address)
+bool UOSCManager::OSCAddressIsValidPath(const FOSCAddress& Address)
 {
-	return Address.IsBundle();
+	return Address.IsValidPath();
 }
 
-bool UOSCManager::OSCAddressIsMessage(const FOSCAddress& Address)
+bool UOSCManager::OSCAddressIsValidPattern(const FOSCAddress& Address)
 {
-	return Address.IsBundle();
-}
-
-bool UOSCManager::OSCAddressIsValid(const FOSCAddress& Address)
-{
-	return Address.IsValid();
+	return Address.IsValidPattern();
 }
 
 FOSCAddress UOSCManager::StringToOSCAddress(const FString& String)
@@ -370,12 +365,17 @@ FOSCAddress UOSCManager::StringToOSCAddress(const FString& String)
 	return FOSCAddress(String);
 }
 
-FOSCAddress& UOSCManager::OSCAddressAppend(FOSCAddress& Address, const FString& ToAppend)
+FOSCAddress& UOSCManager::OSCAddressPushContainer(FOSCAddress& Address, const FString& ToAppend)
 {
-	Address.Append(ToAppend);
+	Address.PushContainer(ToAppend);
 	return Address;
 }
 
+FString UOSCManager::OSCAddressPopContainer(FOSCAddress& Address)
+{
+	return Address.PopContainer();
+}
+
 FOSCAddress UOSCManager::GetOSCMessageAddress(const FOSCMessage& Message)
 {
 	return Message.GetAddress();
@@ -387,9 +387,19 @@ FOSCMessage& UOSCManager::SetOSCMessageAddress(FOSCMessage& Message, const FOSCA
 	return Message;
 }
 
-TArray<FString> UOSCManager::SplitOSCAddress(const FOSCAddress& Address)
+FString UOSCManager::GetOSCAddressContainer(const FOSCAddress& Address, int32 Index)
 {
-	TArray<FString> OutArray;
-	Address.Split(OutArray);
-	return MoveTemp(OutArray);
-}
\ No newline at end of file
+	return Address.GetContainer(Index);
+}
+
+TArray<FString> UOSCManager::GetOSCAddressContainers(const FOSCAddress& Address)
+{
+	TArray<FString> Containers;
+	Address.GetContainers(Containers);
+	return MoveTemp(Containers);
+}
+
+FString UOSCManager::GetOSCAddressMethodName(const FOSCAddress& Address)
+{
+	return Address.GetMethodName();
+}
diff --git a/Engine/Plugins/Runtime/OSC/Source/OSC/Private/OSCMessage.cpp b/Engine/Plugins/Runtime/OSC/Source/OSC/Private/OSCMessage.cpp
--- a/Engine/Plugins/Runtime/OSC/Source/OSC/Private/OSCMessage.cpp
+++ b/Engine/Plugins/Runtime/OSC/Source/OSC/Private/OSCMessage.cpp
@@ -35,9 +35,9 @@ bool FOSCMessage::SetAddress(const FOSCAddress& InAddress)
 {
 	check(Packet.IsValid());
 
-	if (!InAddress.IsMessage())
+	if (!InAddress.IsValidPath())
 	{
-		UE_LOG(LogOSC, Warning, TEXT("Attempting to set invalid OSCAddress '%s'. OSC address must begin with '/'"), *InAddress.Value);
+		UE_LOG(LogOSC, Warning, TEXT("Attempting to set invalid OSCAddress '%s'. OSC address must begin with '/'"), *InAddress.GetFullPath());
 		return false;
 	}
 
diff --git a/Engine/Plugins/Runtime/OSC/Source/OSC/Private/OSCMessagePacket.cpp b/Engine/Plugins/Runtime/OSC/Source/OSC/Private/OSCMessagePacket.cpp
--- a/Engine/Plugins/Runtime/OSC/Source/OSC/Private/OSCMessagePacket.cpp
+++ b/Engine/Plugins/Runtime/OSC/Source/OSC/Private/OSCMessagePacket.cpp
@@ -40,14 +40,14 @@ bool FOSCMessagePacket::IsMessage()
 
 void FOSCMessagePacket::WriteData(FOSCStream& Stream)
 {
-	if (!Address.IsMessage())
+	if (!Address.IsValidPath())
 	{
-		UE_LOG(LogOSC, Warning, TEXT("Failed to write OSCMessagePacket. Invalid OSCAddress '%s'"), *Address.Value);
+		UE_LOG(LogOSC, Warning, TEXT("Failed to write OSCMessagePacket. Invalid OSCAddress '%s'"), *Address.GetFullPath());
 		return;
 	}
 
 	// Begin writing data 
-	Stream.WriteString(Address.Value);
+	Stream.WriteString(Address.GetFullPath());
 
 	// write type tags
 	FString TagTypes = ",";
diff --git a/Engine/Plugins/Runtime/OSC/Source/OSC/Private/OSCPacket.cpp b/Engine/Plugins/Runtime/OSC/Source/OSC/Private/OSCPacket.cpp
--- a/Engine/Plugins/Runtime/OSC/Source/OSC/Private/OSCPacket.cpp
+++ b/Engine/Plugins/Runtime/OSC/Source/OSC/Private/OSCPacket.cpp
@@ -16,19 +16,19 @@ IOSCPacket::~IOSCPacket()
 
 TSharedPtr<IOSCPacket> IOSCPacket::CreatePacket(const uint8* PacketType)
 {
-	const FOSCAddress Address(ANSI_TO_TCHAR((const ANSICHAR*)&PacketType[0]));
-	if (Address.IsMessage())
+	const FString PacketIdentifier(ANSI_TO_TCHAR((const ANSICHAR*)&PacketType[0]));
+	if (PacketIdentifier.StartsWith(OSC::PathSeparator))
 	{
 		return MakeShareable(new FOSCMessagePacket());
 	}
-	else if (Address.IsBundle())
+	else if (PacketIdentifier == OSC::BundleTag)
 	{
 		return MakeShareable(new FOSCBundlePacket());
 	}
 	else
 	{
 		UE_LOG(LogOSC, Warning, TEXT("Failed to parse lead character of OSC message packet. "
-			"Lead identifier of '%c' not valid bundle (#) or message (/) identifier."), PacketType[0]);
+			"Lead identifier of '%c' not valid bundle tag ('%s') or message ('%s') identifier."), *PacketIdentifier, *OSC::BundleTag, *OSC::PathSeparator);
 		return nullptr;
 	}
 }
diff --git a/Engine/Plugins/Runtime/OSC/Source/OSC/Private/OSCServer.cpp b/Engine/Plugins/Runtime/OSC/Source/OSC/Private/OSCServer.cpp
--- a/Engine/Plugins/Runtime/OSC/Source/OSC/Private/OSCServer.cpp
+++ b/Engine/Plugins/Runtime/OSC/Source/OSC/Private/OSCServer.cpp
@@ -61,6 +61,14 @@ void UOSCServer::Callback(const FArrayReaderPtr& Data, const FIPv4Endpoint& Endp
 				{
 					FOSCMessage Message(FOSCMessage(StaticCastSharedPtr<FOSCMessagePacket>(Packet)));
 					OnOscReceived.Broadcast(Message);
+
+					for (const TPair<uint32, FOSCAddress>& Pair: AddressPatterns)
+					{
+						if (Pair.Value.Matches(Message.GetAddress()))
+						{
+							OnOscDispatched.Broadcast(Pair.Value, Message);
+						}
+					}
 				}
 				else if (Packet->IsBundle())
 				{
@@ -69,7 +77,7 @@ void UOSCServer::Callback(const FArrayReaderPtr& Data, const FIPv4Endpoint& Endp
 				}
 				else
 				{
-					UE_LOG(LogOSC, Warning, TEXT("Failed to parse invalid received OSC message. OSCAddress '%s' is invalid."), *Packet->GetAddress().Value);
+					UE_LOG(LogOSC, Warning, TEXT("Failed to parse invalid received OSC message. OSCAddress '%s' is invalid."), *Packet->GetAddress().GetFullPath());
 				}
 			}
 		}),
@@ -221,3 +229,35 @@ TSet<FString> UOSCServer::GetWhitelistedClients() const
 
 	return OutWhitelist;
 }
+
+void UOSCServer::AddOSCAddressPattern(const FOSCAddress& OSCAddressPattern)
+{
+	if (OSCAddressPattern.IsValidPattern())
+	{
+		AddressPatterns.Add(OSCAddressPattern.GetHash(), OSCAddressPattern);
+	}
+}
+
+void UOSCServer::RemoveOSCAddressPattern(const FOSCAddress& OSCAddressPattern)
+{
+	if (OSCAddressPattern.IsValidPattern())
+	{
+		AddressPatterns.Remove(OSCAddressPattern.GetHash());
+	}
+}
+
+
+TArray<FOSCAddress> UOSCServer::GetOSCAddressPatterns() const
+{
+	TArray<FOSCAddress> Addresses;
+	for (const TPair<uint32, FOSCAddress>& Pair : AddressPatterns)
+	{
+		Addresses.Add(Pair.Value);
+	}
+	return MoveTemp(Addresses);
+}
+
+void UOSCServer::ClearOSCAddressPatterns()
+{
+	AddressPatterns.Reset();
+}
\ No newline at end of file
diff --git a/Engine/Plugins/Runtime/OSC/Source/OSC/Public/OSCAddress.h b/Engine/Plugins/Runtime/OSC/Source/OSC/Public/OSCAddress.h
--- a/Engine/Plugins/Runtime/OSC/Source/OSC/Public/OSCAddress.h
+++ b/Engine/Plugins/Runtime/OSC/Source/OSC/Public/OSCAddress.h
@@ -11,8 +11,10 @@
 
 namespace OSC
 {
-	const FString BundleTag		= TEXT("#bundle");
-	const FString PathSeparator = TEXT("/");
+	const FString		BundleTag		= TEXT("#bundle");
+	const FString		PathSeparator	= TEXT("/");
+	const TArray<TCHAR>	InvalidChars	= { ' ', '#', };
+	const TArray<TCHAR>	PatternChars	= { ',', '*', '/', '?', '[', ']', '{', '}' };
 } // namespace OSC
 
 USTRUCT(BlueprintType)
@@ -20,61 +22,39 @@ struct OSC_API FOSCAddress
 {
 	GENERATED_USTRUCT_BODY()
 
-public:
-	FOSCAddress()
-	{
-		Value = OSC::PathSeparator;
-	}
+private:
+	UPROPERTY(Transient)
+	TArray<FString> Containers;
 
-	FOSCAddress(const FString& InValue)
-		: Value(InValue)
-	{
-		if (!IsBundle())
-		{
-			if (!Value.StartsWith(OSC::PathSeparator))
-			{
-				Value = OSC::PathSeparator + Value;
-			}
-		}
-	}
+	UPROPERTY(Transient)
+	FString Method;
 
-	UPROPERTY(BlueprintReadWrite, Category = "Audio|OSC")
-	FString Value;
+	bool bIsValidPattern;
+	bool bIsValidPath;
+	uint32 Hash;
 
-	FORCEINLINE bool IsBundle() const
-	{
-		return Value.StartsWith(OSC::BundleTag);
-	}
+	static bool IsValidPattern(const TArray<FString>& InContainers, const FString& InMethod);
+	static bool IsValidPath(const FString& Path, bool bInvalidateSeparator);
 
-	FORCEINLINE bool IsMessage() const
-	{
-		return Value.StartsWith(OSC::PathSeparator);
-	}
+public:
+	uint32 GetHash() const;
 
-	FORCEINLINE bool IsValid() const
-	{
-		return IsMessage() || IsBundle();
-	}
+	bool IsValidPattern() const;
+	bool IsValidPath() const;
+	bool Matches(const FOSCAddress& Address) const;
 
-	FORCEINLINE void Append(const FString& ToAppend)
-	{
-		if (!IsMessage())
-		{
-			UE_LOG(LogOSC, Warning, TEXT("Append failed for OSCAddress '%s'. Address is either invalid/non-message type."), *Value);
-			return;
-		}
+	void PushContainer(const FString& Container);
 
-		Value.PathAppend(*ToAppend, ToAppend.Len());
-	}
+	void ClearContainers(const FString& Container);
+	void SetMethod(const FString& InMethod);
+	FString PopContainer();
 
-	FORCEINLINE void Split(TArray<FString>& OutArray) const
-	{
-		if (!IsMessage())
-		{
-			UE_LOG(LogOSC, Warning, TEXT("Split failed for OSCAddress '%s'. Address is invalid/non-message type."), *Value);
-			return;
-		}
+	FString GetContainerPath() const;
+	FString GetContainer(int32 Index) const;
+	void GetContainers(TArray<FString>& OutContainers) const;
+	FString GetFullPath() const;
+	const FString& GetMethodName() const;
 
-		Value.ParseIntoArray(OutArray, *OSC::PathSeparator, true);
-	}
+	FOSCAddress();
+	FOSCAddress(const FString& InValue);
 };
diff --git a/Engine/Plugins/Runtime/OSC/Source/OSC/Public/OSCManager.h b/Engine/Plugins/Runtime/OSC/Source/OSC/Public/OSCManager.h
--- a/Engine/Plugins/Runtime/OSC/Source/OSC/Public/OSCManager.h
+++ b/Engine/Plugins/Runtime/OSC/Source/OSC/Public/OSCManager.h
@@ -92,19 +92,19 @@ public:
 	static void GetBlob(const FOSCMessage& Message, const int32 Index, TArray<uint8>& Value);
 
 	UFUNCTION(BlueprintCallable, Category = "Audio|OSC")
-	static bool OSCAddressIsBundle(const FOSCAddress& Address);
+	static bool OSCAddressIsValidPath(const FOSCAddress& Address);
 
 	UFUNCTION(BlueprintCallable, Category = "Audio|OSC")
-	static bool OSCAddressIsMessage(const FOSCAddress& Address);
+	static bool OSCAddressIsValidPattern(const FOSCAddress& Address);
 
 	UFUNCTION(BlueprintCallable, Category = "Audio|OSC")
-	static bool OSCAddressIsValid(const FOSCAddress& Address);
+	static FOSCAddress StringToOSCAddress(const FString& String);
 
 	UFUNCTION(BlueprintCallable, Category = "Audio|OSC")
-	static FOSCAddress StringToOSCAddress(const FString& String);
+	static UPARAM(DisplayName = "Address") FOSCAddress& OSCAddressPushContainer(UPARAM(ref) FOSCAddress& Address, const FString& Container);
 
 	UFUNCTION(BlueprintCallable, Category = "Audio|OSC")
-	static UPARAM(DisplayName = "Address") FOSCAddress& OSCAddressAppend(UPARAM(ref) FOSCAddress& Address, const FString& ToAppend);
+	static UPARAM(DisplayName = "Address") FString OSCAddressPopContainer(UPARAM(ref) FOSCAddress& Address);
 
 	UFUNCTION(BlueprintCallable, Category = "Audio|OSC")
 	static FOSCAddress GetOSCMessageAddress(const FOSCMessage& Message);
@@ -113,5 +113,11 @@ public:
 	static UPARAM(DisplayName = "Message") FOSCMessage& SetOSCMessageAddress(UPARAM(ref) FOSCMessage& Message, const FOSCAddress& Address);
 
 	UFUNCTION(BlueprintCallable, Category = "Audio|OSC")
-	static TArray<FString> SplitOSCAddress(const FOSCAddress& Address);
+	static FString GetOSCAddressContainer(const FOSCAddress& Address, const int32 Index);
+
+	UFUNCTION(BlueprintCallable, Category = "Audio|OSC")
+	static TArray<FString> GetOSCAddressContainers(const FOSCAddress& Address);
+
+	UFUNCTION(BlueprintCallable, Category = "Audio|OSC")
+	static FString GetOSCAddressMethodName(const FOSCAddress& Address);
 };
diff --git a/Engine/Plugins/Runtime/OSC/Source/OSC/Public/OSCServer.h b/Engine/Plugins/Runtime/OSC/Source/OSC/Public/OSCServer.h
--- a/Engine/Plugins/Runtime/OSC/Source/OSC/Public/OSCServer.h
+++ b/Engine/Plugins/Runtime/OSC/Source/OSC/Public/OSCServer.h
@@ -15,11 +15,13 @@
 #include "OSCServer.generated.h"
 
 
-DECLARE_DYNAMIC_MULTICAST_DELEGATE_OneParam(FOSCReceivedMessageEvent, const FOSCMessage &, Message);
-DECLARE_DYNAMIC_MULTICAST_DELEGATE_OneParam(FOSCReceivedBundleEvent, const FOSCBundle &, Bundle);
+DECLARE_DYNAMIC_MULTICAST_DELEGATE_OneParam(FOSCReceivedMessageEvent, const FOSCMessage&, Message);
+DECLARE_DYNAMIC_MULTICAST_DELEGATE_TwoParams(FOSCDispatchMessageEvent, const FOSCAddress&, AddressPattern, const FOSCMessage&, Message);
+DECLARE_DYNAMIC_MULTICAST_DELEGATE_OneParam(FOSCReceivedBundleEvent, const FOSCBundle&, Bundle);
 
 class FSocket;
 
+
 UCLASS(BlueprintType)
 class OSC_API UOSCServer : public UObject
 {
@@ -54,6 +56,11 @@ public:
 	UPROPERTY(BlueprintAssignable, Category = "Audio|OSC")
 	FOSCReceivedMessageEvent OnOscReceived;
 
+	/** Event that gets called when an OSC message is dispatched
+	  * via passing OSCMethod filter. */
+	UPROPERTY(BlueprintAssignable, Category = "Audio|OSC")
+	FOSCDispatchMessageEvent OnOscDispatched;
+
 	/** Event that gets called when an OSC bundle is received. */
 	UPROPERTY(BlueprintAssignable, Category = "Audio|OSC")
 	FOSCReceivedBundleEvent OnOscBundleReceived;
@@ -78,6 +85,22 @@ public:
 	UFUNCTION(BlueprintCallable, Category = "Audio|OSC")
 	TSet<FString> GetWhitelistedClients() const;
 
+	/** Adds OSCAddressPattern to dispatch. */
+	UFUNCTION(BlueprintCallable, Category = "Audio|OSC")
+	void AddOSCAddressPattern(const FOSCAddress& OSCAddressPattern);
+
+	/** Adds OSCAddressPattern to dispatch. */
+	UFUNCTION(BlueprintCallable, Category = "Audio|OSC")
+	void RemoveOSCAddressPattern(const FOSCAddress& OSCAddressPattern);
+
+	/** Returns set of OSCAddressPatterns to dispatch. */
+	UFUNCTION(BlueprintCallable, Category = "Audio|OSC")
+	TArray<FOSCAddress> GetOSCAddressPatterns() const;
+
+	/** Returns set of OSCAddressPatterns to dispatch. */
+	UFUNCTION(BlueprintCallable, Category = "Audio|OSC")
+	void ClearOSCAddressPatterns();
+
 protected:
 	void BeginDestroy() override;
 	
@@ -107,4 +130,7 @@ protected:
 
 	/** Whether or not to loopback if address provided is multicast */
 	bool bMulticastLoopback;
+	
+	/** Address patterns to check against when dispatching incoming messages */
+	TMap<uint32, FOSCAddress> AddressPatterns;
 };
