- Clean-up and document public API
- Make OSC address pattern matching only match once and allow multiple delegates to subscribe to match

diff --git a/Engine/Plugins/Runtime/OSC/Source/OSC/Private/OSCAddress.cpp b/Engine/Plugins/Runtime/OSC/Source/OSC/Private/OSCAddress.cpp
--- a/Engine/Plugins/Runtime/OSC/Source/OSC/Private/OSCAddress.cpp
+++ b/Engine/Plugins/Runtime/OSC/Source/OSC/Private/OSCAddress.cpp
@@ -95,11 +95,6 @@ bool FOSCAddress::Matches(const FOSCAddress& Address) const
 	}
 }
 
-uint32 FOSCAddress::GetHash() const
-{
-	return Hash;
-}
-
 bool FOSCAddress::IsValidPattern() const
 {
 	return bIsValidPattern;
@@ -124,12 +119,28 @@ void FOSCAddress::PushContainer(const FString& Container)
 	Hash = GetTypeHash(GetFullPath());
 }
 
+FString FOSCAddress::PopContainer()
+{
+	if (Containers.Num() > 0)
+	{
+		FString Popped = Containers.Pop(false);
+		Hash = GetTypeHash(GetFullPath());
+	}
+
+	return FString();
+}
+
 void FOSCAddress::ClearContainers(const FString& Container)
 {
 	Containers.Reset();
 	Hash = GetTypeHash(GetFullPath());
 }
 
+const FString& FOSCAddress::GetMethod() const
+{
+	return Method;
+}
+
 void FOSCAddress::SetMethod(const FString& InMethod)
 {
 	const bool bInvalidateSeparator = true;
@@ -144,17 +155,6 @@ void FOSCAddress::SetMethod(const FString& InMethod)
 	Hash = GetTypeHash(GetFullPath());
 }
 
-FString FOSCAddress::PopContainer()
-{
-	if (Containers.Num() > 0)
-	{
-		FString Popped = Containers.Pop(false);
-		Hash = GetTypeHash(GetFullPath());
-	}
-
-	return FString();
-}
-
 FString FOSCAddress::GetContainerPath() const
 {
 	return OSC::PathSeparator + FString::Join(Containers, *OSC::PathSeparator);
@@ -162,7 +162,7 @@ FString FOSCAddress::GetContainerPath() const
 
 FString FOSCAddress::GetContainer(int32 Index) const
 {
-	if (Index < Containers.Num())
+	if (Index >= 0 && Index < Containers.Num())
 	{
 		return Containers[Index];
 	}
@@ -184,8 +184,3 @@ FString FOSCAddress::GetFullPath() const
 
 	return GetContainerPath() + OSC::PathSeparator + Method;
 }
-
-const FString& FOSCAddress::GetMethodName() const
-{
-	return Method;
-}
diff --git a/Engine/Plugins/Runtime/OSC/Source/OSC/Private/OSCManager.cpp b/Engine/Plugins/Runtime/OSC/Source/OSC/Private/OSCManager.cpp
--- a/Engine/Plugins/Runtime/OSC/Source/OSC/Private/OSCManager.cpp
+++ b/Engine/Plugins/Runtime/OSC/Source/OSC/Private/OSCManager.cpp
@@ -97,22 +97,26 @@ UOSCClient* UOSCManager::CreateOSCClient(FString SendIPAddress, int32 Port)
 	return NewOSCClient;
 }
 
-void UOSCManager::ClearMessage(FOSCMessage& Message)
+FOSCMessage& UOSCManager::ClearMessage(FOSCMessage& Message)
 {
 	const TSharedPtr<FOSCMessagePacket>& Packet = StaticCastSharedPtr<FOSCMessagePacket>(Message.GetPacket());
 	if (Packet.IsValid())
 	{
 		Packet->GetArguments().Reset();
 	}
+
+	return Message;
 }
 
-void UOSCManager::ClearBundle(FOSCBundle& Bundle)
+FOSCBundle& UOSCManager::ClearBundle(FOSCBundle& Bundle)
 {
 	const TSharedPtr<FOSCBundlePacket>& Packet = StaticCastSharedPtr<FOSCBundlePacket>(Bundle.GetPacket());
 	if (Packet.IsValid())
 	{
 		Packet->GetPackets().Reset();
 	}
+
+	return Bundle;
 }
 
 FOSCBundle& UOSCManager::AddMessageToBundle(const FOSCMessage& Message, FOSCBundle& Bundle)
@@ -360,7 +364,7 @@ bool UOSCManager::OSCAddressIsValidPattern(const FOSCAddress& Address)
 	return Address.IsValidPattern();
 }
 
-FOSCAddress UOSCManager::StringToOSCAddress(const FString& String)
+FOSCAddress UOSCManager::ConvertStringToOSCAddress(const FString& String)
 {
 	return FOSCAddress(String);
 }
@@ -399,7 +403,23 @@ TArray<FString> UOSCManager::GetOSCAddressContainers(const FOSCAddress& Address)
 	return MoveTemp(Containers);
 }
 
-FString UOSCManager::GetOSCAddressMethodName(const FOSCAddress& Address)
+FString UOSCManager::GetOSCAddressContainerPath(const FOSCAddress& Address)
+{
+	return Address.GetContainerPath();
+}
+
+FString UOSCManager::GetOSCAddressFullPath(const FOSCAddress& Address)
 {
-	return Address.GetMethodName();
+	return Address.GetFullPath();
+}
+
+FString UOSCManager::GetOSCAddressMethod(const FOSCAddress& Address)
+{
+	return Address.GetMethod();
+}
+
+FOSCAddress& UOSCManager::SetOSCAddressMethod(FOSCAddress& Address, const FString& Method)
+{
+	Address.SetMethod(Method);
+	return Address;
 }
diff --git a/Engine/Plugins/Runtime/OSC/Source/OSC/Private/OSCServer.cpp b/Engine/Plugins/Runtime/OSC/Source/OSC/Private/OSCServer.cpp
--- a/Engine/Plugins/Runtime/OSC/Source/OSC/Private/OSCServer.cpp
+++ b/Engine/Plugins/Runtime/OSC/Source/OSC/Private/OSCServer.cpp
@@ -30,12 +30,6 @@ UOSCServer::UOSCServer(const FObjectInitializer& ObjectInitializer)
 
 void UOSCServer::Callback(const FArrayReaderPtr& Data, const FIPv4Endpoint& Endpoint)
 {
-	// Throw request on the ground if not endpoint address not whitelisted.
-	if (bWhitelistClients && !ClientWhitelist.Contains(Endpoint.Address.Value))
-	{
-		return;
-	}
-
 	TSharedPtr<IOSCPacket> Packet = IOSCPacket::CreatePacket(Data->GetData());
 	if (Packet.IsValid())
 	{
@@ -49,11 +43,18 @@ void UOSCServer::Callback(const FArrayReaderPtr& Data, const FIPv4Endpoint& Endp
 		!OSCPackets.IsEmpty())
 	{
 		check(OSCBufferLock == OSC_SERVER_BUFFER_LOCKED);
-		FSimpleDelegateGraphTask::CreateAndDispatchWhenReady(FSimpleDelegateGraphTask::FDelegate::CreateLambda([this]()
+		FSimpleDelegateGraphTask::CreateAndDispatchWhenReady(FSimpleDelegateGraphTask::FDelegate::CreateLambda([this, Endpoint]()
 		{
 			check(OSCBufferLock == OSC_SERVER_BUFFER_LOCKED);
 			FPlatformAtomics::InterlockedCompareExchange(&OSCBufferLock, OSC_SERVER_BUFFER_FREE, OSC_SERVER_BUFFER_LOCKED);
-				
+
+			// Throw request on the ground if endpoint address not whitelisted.
+			if (bWhitelistClients && !ClientWhitelist.Contains(Endpoint.Address.Value))
+			{
+				OSCPackets.Empty();
+				return;
+			}
+
 			TSharedPtr<IOSCPacket> Packet;
 			while (OSCPackets.Dequeue(Packet))
 			{
@@ -61,12 +62,18 @@ void UOSCServer::Callback(const FArrayReaderPtr& Data, const FIPv4Endpoint& Endp
 				{
 					FOSCMessage Message(FOSCMessage(StaticCastSharedPtr<FOSCMessagePacket>(Packet)));
 					OnOscReceived.Broadcast(Message);
+					UE_LOG(LogOSC, Verbose, TEXT("Message received from endpoint '%s', OSCAddress of '%s'."), *Endpoint.Address.ToString(), *Packet->GetAddress().GetFullPath());
 
-					for (const TPair<uint32, FOSCAddress>& Pair: AddressPatterns)
+					for (const TPair<FOSCAddress, FOSCDispatchMessageEvent>& Pair: AddressPatterns)
 					{
-						if (Pair.Value.Matches(Message.GetAddress()))
+						const FOSCDispatchMessageEvent& DispatchEvent = Pair.Value;
+						if (Pair.Key.Matches(Message.GetAddress()))
 						{
-							OnOscDispatched.Broadcast(Pair.Value, Message);
+							DispatchEvent.Broadcast(Pair.Key, Message);
+							UE_LOG(LogOSC, Verbose, TEXT("Message dispatched from endpoint '%s', OSCAddress path of '%s' matched OSCAddress pattern '%s'."),
+								*Endpoint.Address.ToString(),
+								*Packet->GetAddress().GetFullPath(),
+								*Pair.Key.GetFullPath());
 						}
 					}
 				}
@@ -230,34 +237,49 @@ TSet<FString> UOSCServer::GetWhitelistedClients() const
 	return OutWhitelist;
 }
 
-void UOSCServer::AddOSCAddressPattern(const FOSCAddress& OSCAddressPattern)
+void UOSCServer::BindEventToOnOSCAddressPatternMatchesPath(const FOSCAddress& OSCAddressPattern, const FOSCDispatchMessageEventBP& Event)
 {
 	if (OSCAddressPattern.IsValidPattern())
 	{
-		AddressPatterns.Add(OSCAddressPattern.GetHash(), OSCAddressPattern);
+		FOSCDispatchMessageEvent& MessageEvent = AddressPatterns.FindOrAdd(OSCAddressPattern);
+		MessageEvent.AddUnique(Event);
 	}
 }
 
-void UOSCServer::RemoveOSCAddressPattern(const FOSCAddress& OSCAddressPattern)
+void UOSCServer::UnbindEventFromOnOSCAddressPatternMatchesPath(const FOSCAddress& OSCAddressPattern, const FOSCDispatchMessageEventBP& Event)
 {
 	if (OSCAddressPattern.IsValidPattern())
 	{
-		AddressPatterns.Remove(OSCAddressPattern.GetHash());
+		if (FOSCDispatchMessageEvent* AddressPatternEvent = AddressPatterns.Find(OSCAddressPattern))
+		{
+			AddressPatternEvent->Remove(Event);
+			if (!AddressPatternEvent->IsBound())
+			{
+				AddressPatterns.Remove(OSCAddressPattern);
+			}
+		}
 	}
 }
 
-
-TArray<FOSCAddress> UOSCServer::GetOSCAddressPatterns() const
+void UOSCServer::UnbindAllEventsFromOnOSCAddressPatternMatchesPath(const FOSCAddress& OSCAddressPattern)
 {
-	TArray<FOSCAddress> Addresses;
-	for (const TPair<uint32, FOSCAddress>& Pair : AddressPatterns)
+	if (OSCAddressPattern.IsValidPattern())
 	{
-		Addresses.Add(Pair.Value);
+		AddressPatterns.Remove(OSCAddressPattern);
 	}
-	return MoveTemp(Addresses);
 }
 
-void UOSCServer::ClearOSCAddressPatterns()
+void UOSCServer::UnbindAllEventsFromOnOSCAddressPatternMatching()
 {
 	AddressPatterns.Reset();
-}
\ No newline at end of file
+}
+
+TArray<FOSCAddress> UOSCServer::GetBoundOSCAddressPatterns() const
+{
+	TArray<FOSCAddress> OutAddressPatterns;
+	for (const TPair<FOSCAddress, FOSCDispatchMessageEvent>& Pair : AddressPatterns)
+	{
+		OutAddressPatterns.Add(Pair.Key);
+	}
+	return MoveTemp(OutAddressPatterns);
+}
diff --git a/Engine/Plugins/Runtime/OSC/Source/OSC/Public/OSCAddress.h b/Engine/Plugins/Runtime/OSC/Source/OSC/Public/OSCAddress.h
--- a/Engine/Plugins/Runtime/OSC/Source/OSC/Public/OSCAddress.h
+++ b/Engine/Plugins/Runtime/OSC/Source/OSC/Public/OSCAddress.h
@@ -23,12 +23,15 @@ struct OSC_API FOSCAddress
 	GENERATED_USTRUCT_BODY()
 
 private:
+	// Ordered array of container names
 	UPROPERTY(Transient)
 	TArray<FString> Containers;
 
+	// Method name of string
 	UPROPERTY(Transient)
 	FString Method;
 
+	/** Cached values for validity and hash */
 	bool bIsValidPattern;
 	bool bIsValidPath;
 	uint32 Hash;
@@ -37,24 +40,52 @@ private:
 	static bool IsValidPath(const FString& Path, bool bInvalidateSeparator);
 
 public:
-	uint32 GetHash() const;
+	FOSCAddress();
+	FOSCAddress(const FString& InValue);
 
+	/** Returns whether address is valid pattern */
 	bool IsValidPattern() const;
+
+	/** Returns whether address is valid path */
 	bool IsValidPath() const;
+
+	/** Returns whether this address matches the passed address */
 	bool Matches(const FOSCAddress& Address) const;
 
+	/** Pushes container onto address' ordered array of containers */
 	void PushContainer(const FString& Container);
 
+	/** Pops container from ordered array of containers */
+	FString PopContainer();
+
+	/** Clears ordered array of containers */
 	void ClearContainers(const FString& Container);
+
+	/** Get method name of address */
+	const FString& GetMethod() const;
+
+	/** Sets the method name of address */
 	void SetMethod(const FString& InMethod);
-	FString PopContainer();
 
+	/** Returns container path of OSC address in the form '/Container1/Container2' */
 	FString GetContainerPath() const;
+
+	/** Returns container at provided Index.  If Index is out-of-bounds, returns empty string. */
 	FString GetContainer(int32 Index) const;
+
+	/** Builds referenced array of address of containers in order */
 	void GetContainers(TArray<FString>& OutContainers) const;
+
+	/** Returns full path of OSC address in the form '/Container1/Container2/Method' */
 	FString GetFullPath() const;
-	const FString& GetMethodName() const;
 
-	FOSCAddress();
-	FOSCAddress(const FString& InValue);
+	bool operator== (const FOSCAddress& InAddress) const
+	{
+		return Hash == InAddress.Hash;
+	}
+
+	friend uint32 GetTypeHash(const FOSCAddress& InAddress)
+	{
+		return GetTypeHash(InAddress.Hash);
+	}
 };
diff --git a/Engine/Plugins/Runtime/OSC/Source/OSC/Public/OSCManager.h b/Engine/Plugins/Runtime/OSC/Source/OSC/Public/OSCManager.h
--- a/Engine/Plugins/Runtime/OSC/Source/OSC/Public/OSCManager.h
+++ b/Engine/Plugins/Runtime/OSC/Source/OSC/Public/OSCManager.h
@@ -25,99 +25,144 @@ public:
 	UFUNCTION(BlueprintCallable, Category = "Audio|OSC")
 	static class UOSCClient* CreateOSCClient(FString SendIPAddress, int32 Port);
 
+	/** Adds provided message packet to bundle. */
 	UFUNCTION(BlueprintCallable, Category = "Audio|OSC")
 	static UPARAM(DisplayName = "Bundle") FOSCBundle& AddMessageToBundle(const FOSCMessage& Message, UPARAM(ref) FOSCBundle& Bundle);
 
+	/** Adds bundle packet to bundle. */
 	UFUNCTION(BlueprintCallable, Category = "Audio|OSC")
 	static UPARAM(DisplayName = "OutBundle") FOSCBundle& AddBundleToBundle(const FOSCBundle& InBundle, UPARAM(ref) FOSCBundle& OutBundle);
 
+	/** Fills array of messages found in bundle. */
 	UFUNCTION(BlueprintCallable, Category = "Audio|OSC")
 	static void GetMessagesFromBundle(const FOSCBundle& Bundle, TArray<FOSCMessage>& Messages);
 
+	/** Clears provided message of all arguments. */
 	UFUNCTION(BlueprintCallable, Category = "Audio|OSC")
-	static void ClearMessage(UPARAM(ref) FOSCMessage& Message);
+	static UPARAM(DisplayName = "Message") FOSCMessage& ClearMessage(UPARAM(ref) FOSCMessage& Message);
 
+	/** Clears provided bundle of all internal messages/bundle packets. */
 	UFUNCTION(BlueprintCallable, Category = "Audio|OSC")
-	static void ClearBundle(UPARAM(ref) FOSCBundle& Bundle);
+	static UPARAM(DisplayName = "Bundle") FOSCBundle& ClearBundle(UPARAM(ref) FOSCBundle& Bundle);
 
+	/** Adds float value to end of OSCMessage */
 	UFUNCTION(BlueprintCallable, Category = "Audio|OSC")
 	static UPARAM(DisplayName = "Message") FOSCMessage& AddFloat(UPARAM(ref) FOSCMessage& Message, float Value);
 
+	/** Adds Int32 value to end of OSCMessage */
 	UFUNCTION(BlueprintCallable, Category = "Audio|OSC")
 	static UPARAM(DisplayName = "Message") FOSCMessage& AddInt32(UPARAM(ref) FOSCMessage& Message, int32 Value);
 
+	/** Adds Int64 value to end of OSCMessage */
 	UFUNCTION(BlueprintCallable, Category = "Audio|OSC")
 	static UPARAM(DisplayName = "Message") FOSCMessage& AddInt64(UPARAM(ref) FOSCMessage& Message, int64 Value);
 
+	/** Adds string value to end of OSCMessage */
 	UFUNCTION(BlueprintCallable, Category = "Audio|OSC")
 	static UPARAM(DisplayName = "Message") FOSCMessage& AddString(UPARAM(ref) FOSCMessage& Message, FString Value);
 
+	/** Adds blob value to end of OSCMessage */
 	UFUNCTION(BlueprintCallable, Category = "Audio|OSC")
 	static UPARAM(DisplayName = "Message") FOSCMessage& AddBlob(UPARAM(ref) FOSCMessage& Message, TArray<uint8>& Value);
 
+	/** Adds boolean value to end of OSCMessage */
 	UFUNCTION(BlueprintCallable, Category = "Audio|OSC")
 	static UPARAM(DisplayName = "Message") FOSCMessage& AddBool(UPARAM(ref) FOSCMessage& Message, bool Value);
 
+	/** Set Value to float at provided Index in OSCMessage if in bounds and type matches */
 	UFUNCTION(BlueprintCallable, Category = "Audio|OSC")
 	static void GetFloat(const FOSCMessage& Message, const int32 Index, float& Value);
 
+	/** Returns all float values in order of received from OSCMessage */
 	UFUNCTION(BlueprintCallable, Category = "Audio|OSC")
 	static void GetAllFloats(const FOSCMessage& Message, TArray<float>& Values);
 
+	/** Set Value to integer at provided Index in OSCMessage if in bounds and type matches */
 	UFUNCTION(BlueprintCallable, Category = "Audio|OSC")
 	static void GetInt32(const FOSCMessage& Message, const int32 Index, int32& Value);
 
+	/** Returns all integer values in order of received from OSCMessage */
 	UFUNCTION(BlueprintCallable, Category = "Audio|OSC")
 	static void GetAllInt32s(const FOSCMessage& Message, TArray<int32>& Values);
 
+	/** Set Value to Int64 at provided Index in OSCMessage if in bounds and type matches */
 	UFUNCTION(BlueprintCallable, Category = "Audio|OSC")
 	static void GetInt64(const FOSCMessage& Message, const int32 Index, int64& Value);
 
+	/** Returns all Int64 values in order of received from OSCMessage */
 	UFUNCTION(BlueprintCallable, Category = "Audio|OSC")
 	static void GetAllInt64s(const FOSCMessage& Message, TArray<int64>& Values);
 
+	/** Set Value to string at provided Index in OSCMessage if in bounds and type matches */
 	UFUNCTION(BlueprintCallable, Category = "Audio|OSC")
 	static void GetString(const FOSCMessage& Message, const int32 Index, FString& Value);
 
+	/** Returns all string values in order of received from OSCMessage */
 	UFUNCTION(BlueprintCallable, Category = "Audio|OSC")
 	static void GetAllStrings(const FOSCMessage& Message, TArray<FString>& Values);
 
+	/** Sets Value to boolean at provided Index from OSCMessage if in bounds and type matches */
 	UFUNCTION(BlueprintCallable, Category = "Audio|OSC")
 	static void GetBool(const FOSCMessage& Message, const int32 Index, bool& Value);
 
+	/** Returns all boolean values in order of received from OSCMessage */
 	UFUNCTION(BlueprintCallable, Category = "Audio|OSC")
 	static void GetAllBools(const FOSCMessage& Message, TArray<bool>& Values);
 
+	/** Sets Value to blob at provided Index from OSCMessage if in bounds and type matches */
 	UFUNCTION(BlueprintCallable, Category = "Audio|OSC")
 	static void GetBlob(const FOSCMessage& Message, const int32 Index, TArray<uint8>& Value);
 
+	/** Returns whether OSCAddress is valid path */
 	UFUNCTION(BlueprintCallable, Category = "Audio|OSC")
 	static bool OSCAddressIsValidPath(const FOSCAddress& Address);
 
+	/** Returns whether OSCAddress is valid pattern to match against */
 	UFUNCTION(BlueprintCallable, Category = "Audio|OSC")
 	static bool OSCAddressIsValidPattern(const FOSCAddress& Address);
 
+	/* Converts string to OSCAddress */
 	UFUNCTION(BlueprintCallable, Category = "Audio|OSC")
-	static FOSCAddress StringToOSCAddress(const FString& String);
+	static UPARAM(DisplayName = "Address") FOSCAddress ConvertStringToOSCAddress(const FString& String);
 
+	/* Pushes container on end of OSCAddress' ordered container array */
 	UFUNCTION(BlueprintCallable, Category = "Audio|OSC")
 	static UPARAM(DisplayName = "Address") FOSCAddress& OSCAddressPushContainer(UPARAM(ref) FOSCAddress& Address, const FString& Container);
 
+	/* Pops container off end of OSCAddress' ordered container array */
 	UFUNCTION(BlueprintCallable, Category = "Audio|OSC")
 	static UPARAM(DisplayName = "Address") FString OSCAddressPopContainer(UPARAM(ref) FOSCAddress& Address);
 
+	/* Returns copy of message's OSCAddress */
 	UFUNCTION(BlueprintCallable, Category = "Audio|OSC")
-	static FOSCAddress GetOSCMessageAddress(const FOSCMessage& Message);
+	static UPARAM(DisplayName = "Address") FOSCAddress GetOSCMessageAddress(const FOSCMessage& Message);
 
+	/** Sets the OSCAdress Address of the provided message */
 	UFUNCTION(BlueprintCallable, Category = "Audio|OSC")
 	static UPARAM(DisplayName = "Message") FOSCMessage& SetOSCMessageAddress(UPARAM(ref) FOSCMessage& Message, const FOSCAddress& Address);
 
+	/** Returns the OSCAddress container at the provided 'Index.' Returns empty string if index is out-of-bounds. */
 	UFUNCTION(BlueprintCallable, Category = "Audio|OSC")
-	static FString GetOSCAddressContainer(const FOSCAddress& Address, const int32 Index);
+	static UPARAM(DisplayName = "Container") FString GetOSCAddressContainer(const FOSCAddress& Address, const int32 Index);
 
+	/** Returns ordered array of OSCAddress containers within OSCAddress' path. */
 	UFUNCTION(BlueprintCallable, Category = "Audio|OSC")
-	static TArray<FString> GetOSCAddressContainers(const FOSCAddress& Address);
+	static UPARAM(DisplayName = "Containers") TArray<FString> GetOSCAddressContainers(const FOSCAddress& Address);
 
+	/** Returns OSCAddress' container path. */
 	UFUNCTION(BlueprintCallable, Category = "Audio|OSC")
-	static FString GetOSCAddressMethodName(const FOSCAddress& Address);
+	static UPARAM(DisplayName = "Path") FString GetOSCAddressContainerPath(const FOSCAddress& Address);
+
+	/** Returns OSCAddress' full path. (Container path plus method name) */
+	UFUNCTION(BlueprintCallable, Category = "Audio|OSC")
+	static UPARAM(DisplayName = "Full Path") FString GetOSCAddressFullPath(const FOSCAddress& Address);
+
+	/** Returns method name of OSCAddress provided */
+	UFUNCTION(BlueprintCallable, Category = "Audio|OSC")
+	static UPARAM(DisplayName = "Method") FString GetOSCAddressMethod(const FOSCAddress& Address);
+
+	/** Sets the method name of the OSCAddress provided */
+	UFUNCTION(BlueprintCallable, Category = "Audio|OSC")
+	static UPARAM(DisplayName = "Address") FOSCAddress& SetOSCAddressMethod(FOSCAddress& Address, const FString& Method);
 };
+
diff --git a/Engine/Plugins/Runtime/OSC/Source/OSC/Public/OSCServer.h b/Engine/Plugins/Runtime/OSC/Source/OSC/Public/OSCServer.h
--- a/Engine/Plugins/Runtime/OSC/Source/OSC/Public/OSCServer.h
+++ b/Engine/Plugins/Runtime/OSC/Source/OSC/Public/OSCServer.h
@@ -14,12 +14,14 @@
 
 #include "OSCServer.generated.h"
 
+// Forward Declarations
+class FSocket;
 
 DECLARE_DYNAMIC_MULTICAST_DELEGATE_OneParam(FOSCReceivedMessageEvent, const FOSCMessage&, Message);
 DECLARE_DYNAMIC_MULTICAST_DELEGATE_TwoParams(FOSCDispatchMessageEvent, const FOSCAddress&, AddressPattern, const FOSCMessage&, Message);
 DECLARE_DYNAMIC_MULTICAST_DELEGATE_OneParam(FOSCReceivedBundleEvent, const FOSCBundle&, Bundle);
 
-class FSocket;
+DECLARE_DYNAMIC_DELEGATE_TwoParams(FOSCDispatchMessageEventBP, const FOSCAddress&, AddressPattern, const FOSCMessage&, Message);
 
 
 UCLASS(BlueprintType)
@@ -56,11 +58,6 @@ public:
 	UPROPERTY(BlueprintAssignable, Category = "Audio|OSC")
 	FOSCReceivedMessageEvent OnOscReceived;
 
-	/** Event that gets called when an OSC message is dispatched
-	  * via passing OSCMethod filter. */
-	UPROPERTY(BlueprintAssignable, Category = "Audio|OSC")
-	FOSCDispatchMessageEvent OnOscDispatched;
-
 	/** Event that gets called when an OSC bundle is received. */
 	UPROPERTY(BlueprintAssignable, Category = "Audio|OSC")
 	FOSCReceivedBundleEvent OnOscBundleReceived;
@@ -85,21 +82,25 @@ public:
 	UFUNCTION(BlueprintCallable, Category = "Audio|OSC")
 	TSet<FString> GetWhitelistedClients() const;
 
-	/** Adds OSCAddressPattern to dispatch. */
+	/** Adds event to dispatch when OSCAddressPattern is matched. */
+	UFUNCTION(BlueprintCallable, Category = "Audio|OSC")
+	void BindEventToOnOSCAddressPatternMatchesPath(const FOSCAddress& OSCAddressPattern, const FOSCDispatchMessageEventBP& Event);
+
+	/** Unbinds specific event from OSCAddress patter */
 	UFUNCTION(BlueprintCallable, Category = "Audio|OSC")
-	void AddOSCAddressPattern(const FOSCAddress& OSCAddressPattern);
+	void UnbindEventFromOnOSCAddressPatternMatchesPath(const FOSCAddress& OSCAddressPattern, const FOSCDispatchMessageEventBP& Event);
 
-	/** Adds OSCAddressPattern to dispatch. */
+	/** Removes OSCAddressPattern from sending dispatch events. */
 	UFUNCTION(BlueprintCallable, Category = "Audio|OSC")
-	void RemoveOSCAddressPattern(const FOSCAddress& OSCAddressPattern);
+	void UnbindAllEventsFromOnOSCAddressPatternMatchesPath(const FOSCAddress& OSCAddressPattern);
 
-	/** Returns set of OSCAddressPatterns to dispatch. */
+	/** Removes all events from OSCAddressPatterns to dispatch. */
 	UFUNCTION(BlueprintCallable, Category = "Audio|OSC")
-	TArray<FOSCAddress> GetOSCAddressPatterns() const;
+	void UnbindAllEventsFromOnOSCAddressPatternMatching();
 
-	/** Returns set of OSCAddressPatterns to dispatch. */
+	/** Returns set of OSCAddressPatterns currently listening for matches to dispatch. */
 	UFUNCTION(BlueprintCallable, Category = "Audio|OSC")
-	void ClearOSCAddressPatterns();
+	TArray<FOSCAddress> GetBoundOSCAddressPatterns() const;
 
 protected:
 	void BeginDestroy() override;
@@ -131,6 +132,6 @@ protected:
 	/** Whether or not to loopback if address provided is multicast */
 	bool bMulticastLoopback;
 	
-	/** Address patterns to check against when dispatching incoming messages */
-	TMap<uint32, FOSCAddress> AddressPatterns;
+	/** Address pattern hash to  to check against when dispatching incoming messages */
+	TMap<FOSCAddress, FOSCDispatchMessageEvent> AddressPatterns;
 };
