- Remove unnecessary Get/SetAddress on base IOSCPacket class.
- Safeguard against calling stop and then attempting to send message after Socket is closed and nulled on OSCClient.
- Move caching hash/IsValidPattern/IsValidPath to CacheAggregates and call anytime OSCAddress properties are modified.
- Use OSCAddress spec pattern matching instead of substring match

diff --git a/Engine/Plugins/Runtime/OSC/Source/OSC/Private/OSCAddress.cpp b/Engine/Plugins/Runtime/OSC/Source/OSC/Private/OSCAddress.cpp
--- a/Engine/Plugins/Runtime/OSC/Source/OSC/Private/OSCAddress.cpp
+++ b/Engine/Plugins/Runtime/OSC/Source/OSC/Private/OSCAddress.cpp
@@ -1,98 +1,482 @@
 // Copyright 1998-2019 Epic Games, Inc. All Rights Reserved.
 #include "OSCAddress.h"
 
-// TODO: Finish implementing full pattern matching instead of just substring parsing
-#define USE_FULL_OSC_PATTERN_MATCH 0
-
-
-FOSCAddress::FOSCAddress()
-	: bIsValidPattern(false)
-	, bIsValidPath(false)
-	, Hash(GetTypeHash(GetFullPath()))
+namespace OSC
 {
-}
+	const FString BundleTag		= TEXT("#bundle");
+	const FString PathSeparator = TEXT("/");
+} // namespace OSC
 
-FOSCAddress::FOSCAddress(const FString& InValue)
-	: bIsValidPattern(false)
-	, bIsValidPath(false)
+namespace
 {
-	bIsValidPath = IsValidPath(InValue, false);
-	InValue.ParseIntoArray(Containers, *OSC::PathSeparator, true);
+	const TArray<TCHAR> AddressInvalidChars = { ' ', '#', };
+	const TArray<TCHAR> AddressPatternChars = { ',', '*', '?', '[', ']', '{', '}' };
 
-	if (Containers.Num() > 0)
+	int32 AddressFindPatternTerminatorIndex(const FString& Pattern, int32 PatternIter, TCHAR Terminator)
 	{
-		Method = Containers.Pop();
+		int32 EndIndex = PatternIter + 1;
+		for (;Pattern[EndIndex] != Terminator; ++EndIndex)
+		{
+			// Should never hit as this function should never be called on invalid pattern
+			check(EndIndex < Pattern.Len());
+		}
+		return EndIndex;
 	}
 
-	Hash = GetTypeHash(InValue);
-	bIsValidPattern = IsValidPattern(Containers, Method);
-}
-
-// TODO: Pattern should check for validity of search strings.
-// ex. '[' doesn't come before closing ']', etc.  This is essentially
-// here as a stub.
-bool FOSCAddress::IsValidPattern(const TArray<FString>& InContainers, const FString& InMethod)
-{
-#if USE_FULL_OSC_PATTERN_MATCH
-	for (const FString& Container : InContainers)
+	bool AddressIsValidPatternPart(const FString& Part)
 	{
-		for (const TCHAR& Char : Container.GetCharArray())
+		bool bInBrackets = false;
+		bool bInBraces = false;
+
+		for (TCHAR Char : Part.GetCharArray())
 		{
 			check(Char != OSC::PathSeparator[0]);
+			if (AddressInvalidChars.Contains(Char))
+			{
+				return false;
+			}
 
-			if (OSC::PatternChars.Contains(Char))
+			switch (Char)
 			{
-				return true;
+			case ',':
+			{
+				if (!bInBraces)
+				{
+					return false;
+				}
+			}
+			break;
+
+			case '?':
+			case '*':
+			{
+				if (bInBraces || bInBrackets)
+				{
+					return false;
+				}
+			}
+			break;
+
+			case '[':
+			{
+				// Nested bracket/brace case
+				if (bInBraces || bInBrackets)
+				{
+					return false;
+				}
+				bInBrackets = true;
+			}
+			break;
+
+			case ']':
+			{
+				// Mixed bracket/brace case
+				if (bInBraces)
+				{
+					return false;
+				}
+				bInBrackets = false;
+			}
+			break;
+
+			case '{':
+			{
+				// Nested bracket/brace case
+				if (bInBraces || bInBrackets)
+				{
+					return false;
+				}
+				bInBraces = true;
+			}
+			break;
+
+			case '}':
+			{
+				// Mixed bracket/brace case
+				if (bInBrackets)
+				{
+					return false;
+				}
+				bInBraces = false;
+			}
+			break;
 			}
 		}
+
+		// Missing close braces/brackets
+		return !bInBraces && !bInBrackets;
 	}
-	return false;
-#else // USE_FULL_OSC_PATTERN_MATCH
-	return true;
-#endif // !USE_FULL_OSC_PATTERN_MATCH
-}
 
-bool FOSCAddress::IsValidPath(const FString& Path, bool bInvalidateSeparator)
-{
-	for (const TCHAR& Char : Path.GetCharArray())
+	bool AddressIsValidPattern(const TArray<FString>& InContainers, const FString& InMethod)
+	{
+		if (InContainers.Num() == 0 || InMethod.IsEmpty())
+		{
+			return false;
+
+		}
+
+		for (const FString& Container : InContainers)
+		{
+			if (!AddressIsValidPatternPart(Container))
+			{
+				return false;
+			}
+		}
+
+		return AddressIsValidPatternPart(InMethod);
+	}
+
+	bool AddressIsValidPath(const FString& Path, bool bInvalidateSeparator)
 	{
-		if (!bInvalidateSeparator)
+		if (Path.IsEmpty() || Path == OSC::PathSeparator)
+		{
+			return false;
+		}
+
+		for (TCHAR Char : Path.GetCharArray())
 		{
 			if (Char == OSC::PathSeparator[0])
 			{
+				if (bInvalidateSeparator)
+				{
+					return false;
+				}
 				continue;
 			}
+
+			if (AddressInvalidChars.Contains(Char))
+			{
+				return false;
+			}
+
+			if (AddressPatternChars.Contains(Char))
+			{
+				return false;
+			}
 		}
 
-		if (OSC::InvalidChars.Contains(Char))
+		return true;
+	}
+
+	bool AddressBracePatternMatches(const FString& Pattern, int32 PatternStartIndex, int32 PatternEndIndex, const FString& Part, int32& PartIter)
+	{
+		// If these checks are hit, function is being called prior to pattern
+		// validation or pattern validation is failing to validate properly.
+		check(PatternEndIndex - PatternStartIndex > 1);
+		check(Pattern[PatternStartIndex] == '{');
+		check(Pattern[PatternEndIndex] == '}');
+
+		// 1. Empty pattern, so always true
+		if (Pattern.Len() == 2)
 		{
-			return false;
+			return true;
 		}
 
-		if (OSC::PatternChars.Contains(Char))
+		// Increment start index to match first slot in strings to match
+		int32 PatternIter = PatternStartIndex + 1;
+		bool bMatches = false;
+		for (int32 PartIndex = PartIter; PartIndex < Part.Len() && PatternIter < PatternEndIndex;)
 		{
-			return false;
+			// End of matching string.
+			if (Pattern[PatternIter] == ',')
+			{
+				// Check to see if succeeded and return state if so.
+				if (bMatches)
+				{
+					PartIter = PartIndex;
+					return true;
+				}
+				// If not, reset part search and 
+				// continue to next Pattern string.
+				PartIndex = PartIter;
+			}
+			// Mark as matching if we're on the path to a match.
+			else if (Pattern[PatternIter] == Part[PartIndex])
+			{
+				bMatches = true;
+				++PartIndex;
+			}
+			// Reset part iterator and fast forward to next pattern string to match
+			else
+			{
+				bMatches = false;
+				PartIndex = PartIter;
+				while (Pattern[PatternIter] != ',' && PatternIter < PatternEndIndex)
+				{
+					++PatternIter;
+				}
+			}
+			++PatternIter;
 		}
+
+		return false;
 	}
 
-	return true;
+	bool AddressBracketPatternMatches(const FString& Pattern, int32 PatternStartIndex, int32 PatternEndIndex, TCHAR MatchChar)
+	{
+		// If these checks are hit, function is being called prior to pattern
+		// validation or pattern validation is failing to validate properly.
+		check(PatternEndIndex - PatternStartIndex > 1);
+		check(Pattern[PatternStartIndex] == '[');
+		check(Pattern[PatternEndIndex] == ']');
+
+		// 1. Empty pattern, so always true
+		if (Pattern.Len() == 2)
+		{
+			return true;
+		}
+
+		// Prime start index and cache whether or not negation of pattern is requested.
+		PatternStartIndex++;
+		bool bNegate = false;
+		if (Pattern[1] == '!')
+		{
+			// Of non-standard form [!]: OSC 1.0 Standard is non-descript about this
+			// scenario, so treat as if this case negates any additional character.
+			if (Pattern.Len() == 3)
+			{
+				return false;
+			}
+
+			PatternStartIndex++;
+			bNegate = true;
+		}
+
+		// 2. Special case where in the form [a-z] or [!a-z]
+		if (Pattern.Len() == PatternStartIndex + 4)
+		{
+			if (Pattern[PatternStartIndex + 1] == '-')
+			{
+				if (MatchChar >= Pattern[PatternStartIndex] && MatchChar <= Pattern[PatternStartIndex + 2])
+				{
+					return !bNegate;
+				}
+				else if (MatchChar >= Pattern[PatternStartIndex + 2] && MatchChar <= Pattern[PatternStartIndex])
+				{
+					return !bNegate;
+				}
+
+				return bNegate;
+			}
+		}
+
+		// 4. Form is either [!abcdef] or [abcdef]
+		for (int32 i = PatternStartIndex; i < PatternEndIndex; ++i)
+		{
+			if (MatchChar == Pattern[i])
+			{
+				return !bNegate;
+			}
+		}
+
+		return bNegate;
+	}
+
+	bool AddressWildPatternMatches(const FString& Pattern, int32& PatternIter, const FString& Part, int32& PartIter)
+	{
+		// Filter out/bump pattern index for case of having sequential wildcards
+		for (; PatternIter + 1 < Pattern.Len(); ++PatternIter)
+		{
+			if (Pattern[PatternIter + 1] != '*')
+			{
+				break;
+			}
+		}
+
+		// Increment to start of next thing to match next test.  Have to always run following test
+		// here as if it fails, the part iter can be incremented, effectively allowing the lead character to be
+		// "consumed" by the wildcard and the following test has to be run again until search is exhausted
+		// or following test is fulfilled.
+		++PatternIter;
+
+		// Reached end of pattern, which is a '*' so no more parsing on this part char required 
+		// (all path part beyond is considered valid)
+		if (PatternIter == Pattern.Len())
+		{
+			return true;
+		}
+
+		const bool bWildLeadsBrackets = Pattern[PatternIter] == '[';
+		const bool bWildLeadsBraces = Pattern[PatternIter] == '{';
+		const TCHAR Terminator = bWildLeadsBrackets ? ']' : '}';
+
+		int32 PatternEndIndex = 0;
+		if (bWildLeadsBrackets || bWildLeadsBraces)
+		{
+			PatternEndIndex = AddressFindPatternTerminatorIndex(Pattern, PatternIter, Terminator);
+		}
+
+		// Continue bumping PartIter until either at end of Part or Part char matches next pattern
+		for (;PartIter < Part.Len(); ++PartIter)
+		{
+			if (bWildLeadsBrackets)
+			{
+				if (AddressBracketPatternMatches(Pattern, PatternIter, PatternEndIndex, Part[PartIter]))
+				{
+					PatternIter = PatternEndIndex;
+					return true;
+				}
+			}
+			else if (bWildLeadsBraces)
+			{
+				// Brace match updates OutPartIter.  Only pass along and update end of brace section
+				// if it succeeds to avoid corrupting initial state for subsequent tests after failing.
+				int32 OutPartIter = PartIter;
+				if (AddressBracePatternMatches(Pattern, PatternIter, PatternEndIndex, Part, OutPartIter))
+				{
+					PatternIter = PatternEndIndex;
+					PartIter = OutPartIter;
+					return true;
+				}
+			}
+			else
+			{
+				if (Pattern[PatternIter] == Part[PartIter] || Pattern[PatternIter] == '?')
+				{
+					return true;
+				}
+			}
+		}
+
+		return false;
+	}
+
+	bool AddressPartsMatch(const FString& Pattern, const FString& Part)
+	{
+		int32 PartIter = 0;
+		for (int32 PatternIter = 0; PatternIter < Pattern.Len() && PartIter < Part.Len();)
+		{
+			// If pattern is at end but part still has characters to process, match fails
+			if (PatternIter >= Pattern.Len() && PartIter < Part.Len())
+			{
+				return false;
+			}
+
+			// If part is at end but pattern still has rules to process, match fails
+			if (PatternIter < Pattern.Len() && PartIter >= Part.Len())
+			{
+				return false;
+			}
+
+			switch (Pattern[PatternIter])
+			{
+				case '?':
+				{
+					// ignore case as any ol' input path character will do
+					PatternIter++;
+					PartIter++;
+				}
+				break;
+
+				case '*':
+				{
+					if (!AddressWildPatternMatches(Pattern, PatternIter, Part, PartIter))
+					{
+						return false;
+					}
+					PatternIter++;
+					PartIter++;
+				}
+				break;
+
+				case '[':
+				{
+					int32 PatternEndIndex = AddressFindPatternTerminatorIndex(Pattern, PatternIter, ']');
+					if (!AddressBracketPatternMatches(Pattern, PatternIter, PatternEndIndex, Part[PartIter]))
+					{
+						return false;
+					}
+					PatternIter = PatternEndIndex + 1;
+					PartIter++;
+				}
+				break;
+
+				case '{':
+				{
+					int32 PatternEndIndex = AddressFindPatternTerminatorIndex(Pattern, PatternIter, '}');
+					if (!AddressBracePatternMatches(Pattern, PatternIter, PatternEndIndex, Part, PartIter))
+					{
+						return false;
+					}
+					PatternIter = PatternEndIndex + 1;
+					PartIter++;
+				}
+				break;
+
+				// All other valid characters direct test
+				default:
+				{
+					if (Pattern[PatternIter] != Part[PartIter])
+					{
+						return false;
+					}
+					PatternIter++;
+					PartIter++;
+				}
+				break;
+			}
+		}
+		return true;
+	}
+} // namespace <>
+
+
+FOSCAddress::FOSCAddress()
+	: bIsValidPattern(false)
+	, bIsValidPath(false)
+	, Hash(GetTypeHash(GetFullPath()))
+{
 }
 
-bool FOSCAddress::Matches(const FOSCAddress& Address) const
+FOSCAddress::FOSCAddress(const FString& InValue)
+	: bIsValidPattern(false)
+	, bIsValidPath(false)
 {
-	if (IsValidPattern())
+	InValue.ParseIntoArray(Containers, *OSC::PathSeparator, true);
+	if (Containers.Num() > 0)
 	{
-#if USE_FULL_OSC_PATTERN_MATCH
-		// TODO: Implement
-		return false;
-#else // USE_FULL_OSC_PATTERN_MATCH
-		return Address.GetFullPath().Contains(GetFullPath(), ESearchCase::CaseSensitive);
-#endif // !USE_FULL_OSC_PATTERN_MATCH
+		Method = Containers.Pop();
 	}
-	else
+
+	CacheAggregates();
+}
+
+void FOSCAddress::CacheAggregates()
+{
+	const bool bInvalidateSeparator = false;
+	const FString FullPath = GetFullPath();
+
+	Hash = GetTypeHash(FullPath);
+	bIsValidPath = AddressIsValidPath(FullPath, bInvalidateSeparator);
+	bIsValidPattern = AddressIsValidPattern(Containers, Method);
+}
+
+bool FOSCAddress::Matches(const FOSCAddress& InAddress) const
+{
+	if (IsValidPattern() && InAddress.IsValidPath())
 	{
-		return IsValidPath() && Address.IsValidPath() && Address.GetFullPath() == GetFullPath();
+		if (Containers.Num() != InAddress.Containers.Num())
+		{
+			return false;
+		}
+
+		if (!AddressPartsMatch(Method, InAddress.Method))
+		{
+			return false;
+		}
+
+		for (int32 i = 0; i < Containers.Num(); ++i)
+		{
+			if (!AddressPartsMatch(Containers[i], InAddress.Containers[i]))
+			{
+				return false;
+			}
+		}
+
+		return true; 
 	}
+
+	return false;
 }
 
 bool FOSCAddress::IsValidPattern() const
@@ -108,15 +492,15 @@ bool FOSCAddress::IsValidPath() const
 void FOSCAddress::PushContainer(const FString& Container)
 {
 	const bool bInvalidateSeparator = true;
-	if (!IsValidPath(Container, bInvalidateSeparator))
+	if (Container.Contains(OSC::PathSeparator))
 	{
 		UE_LOG(LogOSC, Warning, TEXT("Failed to push container on OSCAddress. "
-			"Cannot contain invalid or reserved OSC pattern identifiers."));
+			"Cannot contain OSC path separator '%s'."), *OSC::PathSeparator);
 		return;
 	}
 
 	Containers.Push(Container);
-	Hash = GetTypeHash(GetFullPath());
+	CacheAggregates();
 }
 
 FString FOSCAddress::PopContainer()
@@ -133,7 +517,7 @@ FString FOSCAddress::PopContainer()
 void FOSCAddress::ClearContainers(const FString& Container)
 {
 	Containers.Reset();
-	Hash = GetTypeHash(GetFullPath());
+	CacheAggregates();
 }
 
 const FString& FOSCAddress::GetMethod() const
@@ -143,16 +527,22 @@ const FString& FOSCAddress::GetMethod() const
 
 void FOSCAddress::SetMethod(const FString& InMethod)
 {
-	const bool bInvalidateSeparator = true;
-	if (!IsValidPath(Method, bInvalidateSeparator))
+	if (InMethod.IsEmpty())
+	{
+		UE_LOG(LogOSC, Warning, TEXT("Failed to set OSCAddress method. "
+			"'InMethod' cannot be empty string."));
+		return;
+	}
+
+	if (InMethod.Contains(OSC::PathSeparator))
 	{
-		UE_LOG(LogOSC, Warning, TEXT("Failed to set method on OSCAddress. "
-			"Cannot contain invalid or reserved OSC pattern identifiers."));
+		UE_LOG(LogOSC, Warning, TEXT("Failed to set OSCAddress method. "
+			"Cannot contain OSC path separator '%s'."), *OSC::PathSeparator);
 		return;
 	}
 
 	Method = InMethod;
-	Hash = GetTypeHash(GetFullPath());
+	CacheAggregates();
 }
 
 FString FOSCAddress::GetContainerPath() const
diff --git a/Engine/Plugins/Runtime/OSC/Source/OSC/Private/OSCBundlePacket.cpp b/Engine/Plugins/Runtime/OSC/Source/OSC/Private/OSCBundlePacket.cpp
--- a/Engine/Plugins/Runtime/OSC/Source/OSC/Private/OSCBundlePacket.cpp
+++ b/Engine/Plugins/Runtime/OSC/Source/OSC/Private/OSCBundlePacket.cpp
@@ -55,7 +55,6 @@ void FOSCBundlePacket::ReadData(FOSCStream& Stream)
 	while (!Stream.HasReachedEnd())
 	{
 		TSharedPtr<IOSCPacket> Packet = IOSCPacket::CreatePacket(Stream.GetData());
-
 		if (Packet.IsValid())
 		{
 			Packet->ReadData(Stream);
@@ -78,13 +77,3 @@ bool FOSCBundlePacket::IsMessage()
 {
 	return false;
 }
-
-const FOSCAddress& FOSCBundlePacket::GetAddress() const
-{
-	const static FOSCAddress BundleIdentifier = FOSCAddress(OSC::BundleTag);
-	return BundleIdentifier;
-}
-
-void FOSCBundlePacket::SetAddress(const FOSCAddress& InAddress)
-{
-}
diff --git a/Engine/Plugins/Runtime/OSC/Source/OSC/Private/OSCBundlePacket.h b/Engine/Plugins/Runtime/OSC/Source/OSC/Private/OSCBundlePacket.h
--- a/Engine/Plugins/Runtime/OSC/Source/OSC/Private/OSCBundlePacket.h
+++ b/Engine/Plugins/Runtime/OSC/Source/OSC/Private/OSCBundlePacket.h
@@ -34,10 +34,6 @@ public:
 	  * adding packet data to internal packet bundle. */
 	virtual void ReadData(FOSCStream& Stream) override;
 
-	virtual const FOSCAddress& GetAddress() const override;
-
-	virtual void SetAddress(const FOSCAddress& Address) override;
-
 private:
 	/** Bundle of OSC packets. */
 	FPacketBundle Packets;
diff --git a/Engine/Plugins/Runtime/OSC/Source/OSC/Private/OSCClient.cpp b/Engine/Plugins/Runtime/OSC/Source/OSC/Private/OSCClient.cpp
--- a/Engine/Plugins/Runtime/OSC/Source/OSC/Private/OSCClient.cpp
+++ b/Engine/Plugins/Runtime/OSC/Source/OSC/Private/OSCClient.cpp
@@ -7,6 +7,7 @@
 #include "OSCStream.h"
 #include "OSCBundle.h"
 #include "OSCLog.h"
+#include "OSCMessagePacket.h"
 #include "OSCPacket.h"
 
 
@@ -16,12 +17,9 @@ namespace
 } // namespace <>
 
 
-UOSCClient::UOSCClient()
-	: Socket(FUdpSocketBuilder(*GetName()).Build())
-{
-}
-
-UOSCClient::~UOSCClient()
+UOSCClient::UOSCClient(const FObjectInitializer& ObjectInitializer)
+	: Super(ObjectInitializer)
+	, Socket(FUdpSocketBuilder(*GetName()).Build())
 {
 }
 
@@ -54,13 +52,24 @@ void UOSCClient::GetSendIPAddress(FString& InIPAddress, int32& Port)
 
 void UOSCClient::SendPacket(IOSCPacket& Packet)
 {
-	const FOSCAddress& OSCAddress = Packet.GetAddress();
-	if (Packet.IsMessage() && !OSCAddress.IsValidPath())
+	if (!Socket)
 	{
-		UE_LOG(LogOSC, Warning, TEXT("Failed to write packet data. Invalid OSCAddress '%s'"), *OSCAddress.GetFullPath());
+		UE_LOG(LogOSC, Error, TEXT("Socket has been closed. OSCClient '%s' failed to send msg"), *GetName());
 		return;
 	}
 
+	const FOSCAddress* OSCAddress = nullptr;
+	if (Packet.IsMessage())
+	{
+		const FOSCMessagePacket& MessagePacket = static_cast<FOSCMessagePacket&>(Packet);
+		OSCAddress = &MessagePacket.GetAddress();
+		if (!OSCAddress->IsValidPath())
+		{
+			UE_LOG(LogOSC, Warning, TEXT("Failed to write packet data. Invalid OSCAddress '%s'"), *OSCAddress->GetFullPath());
+			return;
+		}
+	}
+
 	if (IPAddress)
 	{
 		FOSCStream Stream = FOSCStream(OUTPUT_BUFFER_SIZE);
@@ -76,8 +85,8 @@ void UOSCClient::SendPacket(IOSCPacket& Packet)
 			const bool bSuccess = Socket->SendTo(DataPtr, Length, BytesSent, *IPAddress);
 			if (!bSuccess || BytesSent <= 0)
 			{
-				UE_LOG(LogOSC, Verbose, TEXT("OSC Packet failed: Client '%s', OSCAddress '%s', Send IPAddress %s, Attempted Bytes = %d"),
-					*GetName(), *OSCAddress.GetFullPath(), *IPAddress->ToString(true /*bAppendPort*/), AttemptedLength);
+				UE_LOG(LogOSC, Verbose, TEXT("OSC Packet failed: Client '%s', OSC Identifier '%s', Send IPAddress %s, Attempted Bytes = %d"),
+					*GetName(), OSCAddress ? *OSCAddress->GetFullPath() : *OSC::BundleTag, *IPAddress->ToString(true /*bAppendPort*/), AttemptedLength);
 				return;
 			}
 
@@ -85,8 +94,8 @@ void UOSCClient::SendPacket(IOSCPacket& Packet)
 			DataPtr += BytesSent;
 		}
 
-		UE_LOG(LogOSC, Verbose, TEXT("OSC Packet sent: Client '%s', OSCAddress '%s', Send IPAddress %s, Bytes Sent = %d"),
-			*GetName(), *OSCAddress.GetFullPath(), *IPAddress->ToString(true /*bAppendPort*/), AttemptedLength);
+		UE_LOG(LogOSC, Verbose, TEXT("OSC Packet sent: Client '%s', OSC Identifier '%s', Send IPAddress %s, Bytes Sent = %d"),
+			*GetName(), OSCAddress ? *OSCAddress->GetFullPath() : *OSC::BundleTag, *IPAddress->ToString(true /*bAppendPort*/), AttemptedLength);
 	}
 }
 
diff --git a/Engine/Plugins/Runtime/OSC/Source/OSC/Private/OSCMessage.cpp b/Engine/Plugins/Runtime/OSC/Source/OSC/Private/OSCMessage.cpp
--- a/Engine/Plugins/Runtime/OSC/Source/OSC/Private/OSCMessage.cpp
+++ b/Engine/Plugins/Runtime/OSC/Source/OSC/Private/OSCMessage.cpp
@@ -41,12 +41,12 @@ bool FOSCMessage::SetAddress(const FOSCAddress& InAddress)
 		return false;
 	}
 
-	Packet->SetAddress(InAddress);
+	StaticCastSharedPtr<FOSCMessagePacket>(Packet)->SetAddress(InAddress);
 	return true;
 }
 
 const FOSCAddress& FOSCMessage::GetAddress() const
 {
 	check(Packet.IsValid());
-	return Packet->GetAddress();
+	return StaticCastSharedPtr<FOSCMessagePacket>(Packet)->GetAddress();
 }
\ No newline at end of file
diff --git a/Engine/Plugins/Runtime/OSC/Source/OSC/Private/OSCMessagePacket.h b/Engine/Plugins/Runtime/OSC/Source/OSC/Private/OSCMessagePacket.h
--- a/Engine/Plugins/Runtime/OSC/Source/OSC/Private/OSCMessagePacket.h
+++ b/Engine/Plugins/Runtime/OSC/Source/OSC/Private/OSCMessagePacket.h
@@ -17,7 +17,7 @@ public:
 	void SetAddress(const FOSCAddress& InAddress);
 
 	/** Get OSC message address. */
-	virtual const FOSCAddress& GetAddress() const override;
+	virtual const FOSCAddress& GetAddress() const;
 
 	/** Get arguments array. */
 	TArray<FOSCType>& GetArguments();
diff --git a/Engine/Plugins/Runtime/OSC/Source/OSC/Private/OSCServer.cpp b/Engine/Plugins/Runtime/OSC/Source/OSC/Private/OSCServer.cpp
--- a/Engine/Plugins/Runtime/OSC/Source/OSC/Private/OSCServer.cpp
+++ b/Engine/Plugins/Runtime/OSC/Source/OSC/Private/OSCServer.cpp
@@ -34,7 +34,6 @@ void UOSCServer::Callback(const FArrayReaderPtr& Data, const FIPv4Endpoint& Endp
 	if (Packet.IsValid())
 	{
 		FOSCStream Stream = FOSCStream(Data->GetData(), Data->Num());
-
 		Packet->ReadData(Stream);
 		OSCPackets.Enqueue(MoveTemp(Packet));
 	}
@@ -62,7 +61,7 @@ void UOSCServer::Callback(const FArrayReaderPtr& Data, const FIPv4Endpoint& Endp
 				{
 					FOSCMessage Message(FOSCMessage(StaticCastSharedPtr<FOSCMessagePacket>(Packet)));
 					OnOscReceived.Broadcast(Message);
-					UE_LOG(LogOSC, Verbose, TEXT("Message received from endpoint '%s', OSCAddress of '%s'."), *Endpoint.Address.ToString(), *Packet->GetAddress().GetFullPath());
+					UE_LOG(LogOSC, Verbose, TEXT("Message received from endpoint '%s', OSCAddress of '%s'."), *Endpoint.Address.ToString(), *Message.GetAddress().GetFullPath());
 
 					for (const TPair<FOSCAddress, FOSCDispatchMessageEvent>& Pair: AddressPatterns)
 					{
@@ -72,7 +71,7 @@ void UOSCServer::Callback(const FArrayReaderPtr& Data, const FIPv4Endpoint& Endp
 							DispatchEvent.Broadcast(Pair.Key, Message);
 							UE_LOG(LogOSC, Verbose, TEXT("Message dispatched from endpoint '%s', OSCAddress path of '%s' matched OSCAddress pattern '%s'."),
 								*Endpoint.Address.ToString(),
-								*Packet->GetAddress().GetFullPath(),
+								*Message.GetAddress().GetFullPath(),
 								*Pair.Key.GetFullPath());
 						}
 					}
@@ -84,7 +83,7 @@ void UOSCServer::Callback(const FArrayReaderPtr& Data, const FIPv4Endpoint& Endp
 				}
 				else
 				{
-					UE_LOG(LogOSC, Warning, TEXT("Failed to parse invalid received OSC message. OSCAddress '%s' is invalid."), *Packet->GetAddress().GetFullPath());
+					UE_LOG(LogOSC, Warning, TEXT("Failed to parse invalid received OSC message. Invalid type (packet is neither identified as message nor bundle)."));
 				}
 			}
 		}),
diff --git a/Engine/Plugins/Runtime/OSC/Source/OSC/Private/OSCStream.cpp b/Engine/Plugins/Runtime/OSC/Source/OSC/Private/OSCStream.cpp
--- a/Engine/Plugins/Runtime/OSC/Source/OSC/Private/OSCStream.cpp
+++ b/Engine/Plugins/Runtime/OSC/Source/OSC/Private/OSCStream.cpp
@@ -9,7 +9,7 @@ FOSCStream::FOSCStream()
 {
 }
 
-FOSCStream::FOSCStream(const uint8* InData, int32 InSize) 
+FOSCStream::FOSCStream(const uint8* InData, int32 InSize)
 	: Data(InData, InSize)
 	, Position(0)
 {
@@ -346,8 +346,9 @@ void FOSCStream::WriteFloat(float Value)
 FString FOSCStream::ReadString()
 {
 	const int32 DataSize = Data.Num();
-	
 	const int32 InitPosition = Position;
+	check(Position < DataSize);
+
 	const ANSICHAR* StrStart = (ANSICHAR*)(&Data[InitPosition]);
 
 	int32 i = InitPosition;
@@ -376,14 +377,19 @@ FString FOSCStream::ReadString()
 void FOSCStream::WriteString(const FString& InString)
 {
 	const TArray<TCHAR>& CharArr = InString.GetCharArray();
-
 	const int32 Count = CharArr.Num();
-	for (int32 i = 0; i < Count; i++)
+
+	if (Count == 0)
 	{
-		Data[Position + i] = CharArr[i];
+		WriteChar('\0');
+	}
+	else
+	{
+		for (int32 i = 0; i < Count; i++)
+		{
+			WriteChar(CharArr[i]);
+		}
 	}
-
-	Position += Count;
 
 	// Increment & pad string with null terminator
 	const int32 NumPaddingZeros = ((Count + 3) / 4) * 4;
diff --git a/Engine/Plugins/Runtime/OSC/Source/OSC/Public/OSCAddress.h b/Engine/Plugins/Runtime/OSC/Source/OSC/Public/OSCAddress.h
--- a/Engine/Plugins/Runtime/OSC/Source/OSC/Public/OSCAddress.h
+++ b/Engine/Plugins/Runtime/OSC/Source/OSC/Public/OSCAddress.h
@@ -11,10 +11,8 @@
 
 namespace OSC
 {
-	const FString		BundleTag		= TEXT("#bundle");
-	const FString		PathSeparator	= TEXT("/");
-	const TArray<TCHAR>	InvalidChars	= { ' ', '#', };
-	const TArray<TCHAR>	PatternChars	= { ',', '*', '/', '?', '[', ']', '{', '}' };
+	extern const FString	BundleTag;
+	extern const FString	PathSeparator;
 } // namespace OSC
 
 USTRUCT(BlueprintType)
@@ -36,8 +34,7 @@ private:
 	bool bIsValidPath;
 	uint32 Hash;
 
-	static bool IsValidPattern(const TArray<FString>& InContainers, const FString& InMethod);
-	static bool IsValidPath(const FString& Path, bool bInvalidateSeparator);
+	void CacheAggregates();
 
 public:
 	FOSCAddress();
@@ -49,17 +46,19 @@ public:
 	/** Returns whether address is valid path */
 	bool IsValidPath() const;
 
-	/** Returns whether this address matches the passed address */
-	bool Matches(const FOSCAddress& Address) const;
+	/** Returns if this address is a pattern that matches InAddress' path.
+	  * If passed address is not a valid path, returns false.
+	  */
+	bool Matches(const FOSCAddress& InAddress) const;
 
 	/** Pushes container onto address' ordered array of containers */
-	void PushContainer(const FString& Container);
+	void PushContainer(const FString& InContainer);
 
 	/** Pops container from ordered array of containers */
 	FString PopContainer();
 
 	/** Clears ordered array of containers */
-	void ClearContainers(const FString& Container);
+	void ClearContainers(const FString& InContainer);
 
 	/** Get method name of address */
 	const FString& GetMethod() const;
@@ -86,6 +85,6 @@ public:
 
 	friend uint32 GetTypeHash(const FOSCAddress& InAddress)
 	{
-		return GetTypeHash(InAddress.Hash);
+		return InAddress.Hash;
 	}
 };
diff --git a/Engine/Plugins/Runtime/OSC/Source/OSC/Public/OSCClient.h b/Engine/Plugins/Runtime/OSC/Source/OSC/Public/OSCClient.h
--- a/Engine/Plugins/Runtime/OSC/Source/OSC/Public/OSCClient.h
+++ b/Engine/Plugins/Runtime/OSC/Source/OSC/Public/OSCClient.h
@@ -14,11 +14,10 @@ class FSocket;
 UCLASS(BlueprintType)
 class OSC_API UOSCClient : public UObject
 {
-	GENERATED_BODY()
+	GENERATED_UCLASS_BODY()
 
 public:
 	UOSCClient();
-	virtual ~UOSCClient();
 
 	/** Sets the OSC Client IP address and port. Returns whether 
 	  * address and port was successfully set. */
diff --git a/Engine/Plugins/Runtime/OSC/Source/OSC/Public/OSCPacket.h b/Engine/Plugins/Runtime/OSC/Source/OSC/Public/OSCPacket.h
--- a/Engine/Plugins/Runtime/OSC/Source/OSC/Public/OSCPacket.h
+++ b/Engine/Plugins/Runtime/OSC/Source/OSC/Public/OSCPacket.h
@@ -21,12 +21,6 @@ public:
 	/** Read packet data from stream */
 	virtual void ReadData(FOSCStream& Stream) = 0;
 	
-	/** Returns OSC address identifier of packet */
-	virtual const FOSCAddress& GetAddress() const = 0;
-
-	/** Set the address identifier of packet */
-	virtual void SetAddress(const FOSCAddress& Address) = 0;
-
 	/** Returns true if packet is message */
 	virtual bool IsMessage() = 0;
 	
