- Fix package fail in OSC by moving networking implementation to internal proxies for Client/Server code & cleaning up build.cs
- Fix crash in packaged builds on initialization by moving socket connection for Server/Client to separate function (as opposed to being created in uobject construction)
- Minor clean-up to thread code & recurse & dispatch bundles and contained messages/bundles appropriately

diff --git a/Engine/Plugins/Runtime/OSC/Source/OSC/OSC.Build.cs b/Engine/Plugins/Runtime/OSC/Source/OSC/OSC.Build.cs
--- a/Engine/Plugins/Runtime/OSC/Source/OSC/OSC.Build.cs
+++ b/Engine/Plugins/Runtime/OSC/Source/OSC/OSC.Build.cs
@@ -4,26 +4,23 @@ namespace UnrealBuildTool.Rules
 {
 	public class OSC : ModuleRules
 	{
-        public OSC(ReadOnlyTargetRules Target) : base(Target)
+		public OSC(ReadOnlyTargetRules Target) : base(Target)
 		{
-            PublicDependencyModuleNames.AddRange(
+			PublicDependencyModuleNames.AddRange(
 				new string[] {
-                    "Core",
+					"Core",
 					"CoreUObject",
 					"Engine",
-                }
-            );
+				}
+			);
 
-            PrivateDependencyModuleNames.AddRange(
-                new string[] {
-                    "Core",
-                    "CoreUObject",
-                    "Engine",
-                    "Sockets",
-                    "InputCore",
-                    "Networking",
-                }
-            );
-        }
+			PrivateDependencyModuleNames.AddRange(
+				new string[] {
+					"InputCore",
+					"Networking",
+					"Sockets",
+				}
+			);
+		}
 	}
 }
diff --git a/Engine/Plugins/Runtime/OSC/Source/OSC/Private/OSCClient.cpp b/Engine/Plugins/Runtime/OSC/Source/OSC/Private/OSCClient.cpp
--- a/Engine/Plugins/Runtime/OSC/Source/OSC/Private/OSCClient.cpp
+++ b/Engine/Plugins/Runtime/OSC/Source/OSC/Private/OSCClient.cpp
@@ -4,107 +4,43 @@
 #include "Common/UdpSocketBuilder.h"
 #include "Sockets.h"
 
-#include "OSCStream.h"
 #include "OSCBundle.h"
+#include "OSCClientProxy.h"
 #include "OSCLog.h"
-#include "OSCMessagePacket.h"
+#include "OSCMessage.h"
 #include "OSCPacket.h"
-
-
-namespace
-{
-	static const int32 OUTPUT_BUFFER_SIZE = 1024;
-} // namespace <>
+#include "OSCStream.h"
 
 
 UOSCClient::UOSCClient(const FObjectInitializer& ObjectInitializer)
 	: Super(ObjectInitializer)
-	, Socket(FUdpSocketBuilder(*GetName()).Build())
+	, ClientProxy(nullptr)
 {
 }
 
-bool UOSCClient::SetSendIPAddress(const FString& InIPAddress, const int32 Port)
+void UOSCClient::Connect()
 {
-	IPAddress = ISocketSubsystem::Get(PLATFORM_SOCKETSUBSYSTEM)->CreateInternetAddr();
-
-	bool bIsValidAddress = true;
-	IPAddress->SetIp(*InIPAddress, bIsValidAddress);
-	IPAddress->SetPort(Port);
-
-	if (bIsValidAddress)
-	{
-		UE_LOG(LogOSC, Verbose, TEXT("OSCClient '%s' SetSendIpAddress: %s:%d"), *GetName(), *InIPAddress, Port);
-	}
-	else
-	{
-		UE_LOG(LogOSC, Warning, TEXT("OSCClient '%s' SetSendIpAddress Failed for input: %s:%d"), *GetName(), *InIPAddress, Port);
-	}
-
-	return bIsValidAddress;
+	check(!ClientProxy.IsValid());
+	ClientProxy.Reset(new FOSCClientProxy(GetName()));
 }
 
 void UOSCClient::GetSendIPAddress(FString& InIPAddress, int32& Port)
 {
-	const bool bAppendPort = false;
-	InIPAddress = IPAddress->ToString(bAppendPort);
-	Port = IPAddress->GetPort();
+	check(ClientProxy.IsValid());
+	ClientProxy->GetSendIPAddress(InIPAddress, Port);
 }
 
-void UOSCClient::SendPacket(IOSCPacket& Packet)
+bool UOSCClient::SetSendIPAddress(const FString& InIPAddress, const int32 Port)
 {
-	if (!Socket)
-	{
-		UE_LOG(LogOSC, Error, TEXT("Socket has been closed. OSCClient '%s' failed to send msg"), *GetName());
-		return;
-	}
-
-	const FOSCAddress* OSCAddress = nullptr;
-	if (Packet.IsMessage())
-	{
-		const FOSCMessagePacket& MessagePacket = static_cast<FOSCMessagePacket&>(Packet);
-		OSCAddress = &MessagePacket.GetAddress();
-		if (!OSCAddress->IsValidPath())
-		{
-			UE_LOG(LogOSC, Warning, TEXT("Failed to write packet data. Invalid OSCAddress '%s'"), *OSCAddress->GetFullPath());
-			return;
-		}
-	}
-
-	if (IPAddress)
-	{
-		FOSCStream Stream = FOSCStream(OUTPUT_BUFFER_SIZE);
-		const uint8* DataPtr = Stream.GetData();
-		Packet.WriteData(Stream);
-
-		int32 BytesSent = 0;
-
-		const int32 AttemptedLength = Stream.GetPosition();
-		int32 Length = AttemptedLength;
-		while (Length > 0)
-		{
-			const bool bSuccess = Socket->SendTo(DataPtr, Length, BytesSent, *IPAddress);
-			if (!bSuccess || BytesSent <= 0)
-			{
-				UE_LOG(LogOSC, Verbose, TEXT("OSC Packet failed: Client '%s', OSC Identifier '%s', Send IPAddress %s, Attempted Bytes = %d"),
-					*GetName(), OSCAddress ? *OSCAddress->GetFullPath() : *OSC::BundleTag, *IPAddress->ToString(true /*bAppendPort*/), AttemptedLength);
-				return;
-			}
-
-			Length -= BytesSent;
-			DataPtr += BytesSent;
-		}
-
-		UE_LOG(LogOSC, Verbose, TEXT("OSC Packet sent: Client '%s', OSC Identifier '%s', Send IPAddress %s, Bytes Sent = %d"),
-			*GetName(), OSCAddress ? *OSCAddress->GetFullPath() : *OSC::BundleTag, *IPAddress->ToString(true /*bAppendPort*/), AttemptedLength);
-	}
+	check(ClientProxy.IsValid());
+	return ClientProxy->SetSendIPAddress(InIPAddress, Port);
 }
 
 void UOSCClient::Stop()
 {
-	if (Socket)
+	if (ClientProxy.IsValid())
 	{
-		ISocketSubsystem::Get(PLATFORM_SOCKETSUBSYSTEM)->DestroySocket(Socket);
-		Socket = nullptr;
+		ClientProxy->Stop();
 	}
 }
 
@@ -116,14 +52,12 @@ void UOSCClient::BeginDestroy()
 
 void UOSCClient::SendOSCMessage(FOSCMessage& Message)
 {
-	const TSharedPtr<IOSCPacket>& Packet = Message.GetPacket();
-	check(Packet.IsValid());
-	SendPacket(*Packet.Get());
+	check(ClientProxy.IsValid());
+	ClientProxy->SendMessage(Message);
 }
 
 void UOSCClient::SendOSCBundle(FOSCBundle& Bundle)
 {
-	const TSharedPtr<IOSCPacket>& Packet = Bundle.GetPacket();
-	check(Packet.IsValid());
-	SendPacket(*Packet.Get());
+	check(ClientProxy.IsValid());
+	ClientProxy->SendBundle(Bundle);
 }
diff --git a/Engine/Plugins/Runtime/OSC/Source/OSC/Private/OSCClientProxy.cpp b/Engine/Plugins/Runtime/OSC/Source/OSC/Private/OSCClientProxy.cpp
new file mode 100644
--- /dev/null
+++ b/Engine/Plugins/Runtime/OSC/Source/OSC/Private/OSCClientProxy.cpp
@@ -0,0 +1,121 @@
+// Copyright 1998-2019 Epic Games, Inc. All Rights Reserved.
+#include "OSCClientProxy.h"
+
+#include "Common/UdpSocketBuilder.h"
+#include "Sockets.h"
+
+#include "OSCBundlePacket.h"
+#include "OSCLog.h"
+#include "OSCClient.h"
+#include "OSCMessagePacket.h"
+
+namespace
+{
+	static const int32 OUTPUT_BUFFER_SIZE = 1024;
+} // namespace <>
+
+
+FOSCClientProxy::FOSCClientProxy(const FString& InClientName)
+	: Socket(FUdpSocketBuilder(*InClientName).Build())
+{
+	ClientName = InClientName;
+}
+
+void FOSCClientProxy::GetSendIPAddress(FString& InIPAddress, int32& Port) const
+{
+	const bool bAppendPort = false;
+	InIPAddress = IPAddress->ToString(bAppendPort);
+	Port = IPAddress->GetPort();
+}
+
+bool FOSCClientProxy::SetSendIPAddress(const FString& InIPAddress, const int32 Port)
+{
+	IPAddress = ISocketSubsystem::Get(PLATFORM_SOCKETSUBSYSTEM)->CreateInternetAddr();
+
+	bool bIsValidAddress = true;
+	IPAddress->SetIp(*InIPAddress, bIsValidAddress);
+	IPAddress->SetPort(Port);
+
+	if (bIsValidAddress)
+	{
+		UE_LOG(LogOSC, Verbose, TEXT("OSCClient '%s' SetSendIpAddress: %s:%d"), *ClientName, *InIPAddress, Port);
+	}
+	else
+	{
+		UE_LOG(LogOSC, Warning, TEXT("OSCClient '%s' SetSendIpAddress Failed for input: %s:%d"), *ClientName, *InIPAddress, Port);
+	}
+
+	return bIsValidAddress;
+}
+
+void FOSCClientProxy::SendPacket(IOSCPacket& Packet)
+{
+	if (!Socket)
+	{
+		UE_LOG(LogOSC, Error, TEXT("Socket has been closed. OSCClient '%s' failed to send msg"), *ClientName);
+		return;
+	}
+
+	const FOSCAddress* OSCAddress = nullptr;
+	if (Packet.IsMessage())
+	{
+		const FOSCMessagePacket& MessagePacket = static_cast<FOSCMessagePacket&>(Packet);
+		OSCAddress = &MessagePacket.GetAddress();
+		if (!OSCAddress->IsValidPath())
+		{
+			UE_LOG(LogOSC, Warning, TEXT("Failed to write packet data. Invalid OSCAddress '%s'"), *OSCAddress->GetFullPath());
+			return;
+		}
+	}
+
+	if (IPAddress)
+	{
+		FOSCStream Stream = FOSCStream(OUTPUT_BUFFER_SIZE);
+		const uint8* DataPtr = Stream.GetData();
+		Packet.WriteData(Stream);
+
+		int32 BytesSent = 0;
+
+		const int32 AttemptedLength = Stream.GetPosition();
+		int32 Length = AttemptedLength;
+		while (Length > 0)
+		{
+			const bool bSuccess = Socket->SendTo(DataPtr, Length, BytesSent, *IPAddress);
+			if (!bSuccess || BytesSent <= 0)
+			{
+				UE_LOG(LogOSC, Verbose, TEXT("OSC Packet failed: Client '%s', OSC Identifier '%s', Send IPAddress %s, Attempted Bytes = %d"),
+					*ClientName, OSCAddress ? *OSCAddress->GetFullPath() : *OSC::BundleTag, *IPAddress->ToString(true /*bAppendPort*/), AttemptedLength);
+				return;
+			}
+
+			Length -= BytesSent;
+			DataPtr += BytesSent;
+		}
+
+		UE_LOG(LogOSC, Verbose, TEXT("OSC Packet sent: Client '%s', OSC Identifier '%s', Send IPAddress %s, Bytes Sent = %d"),
+			*ClientName, OSCAddress ? *OSCAddress->GetFullPath() : *OSC::BundleTag, *IPAddress->ToString(true /*bAppendPort*/), AttemptedLength);
+	}
+}
+
+void FOSCClientProxy::SendMessage(FOSCMessage& Message)
+{
+	const TSharedPtr<IOSCPacket>& Packet = Message.GetPacket();
+	check(Packet.IsValid());
+	SendPacket(*Packet.Get());
+}
+
+void FOSCClientProxy::SendBundle(FOSCBundle& Bundle)
+{
+	const TSharedPtr<IOSCPacket>& Packet = Bundle.GetPacket();
+	check(Packet.IsValid());
+	SendPacket(*Packet.Get());
+}
+
+void FOSCClientProxy::Stop()
+{
+	if (Socket)
+	{
+		ISocketSubsystem::Get(PLATFORM_SOCKETSUBSYSTEM)->DestroySocket(Socket);
+		Socket = nullptr;
+	}
+}
diff --git a/Engine/Plugins/Runtime/OSC/Source/OSC/Private/OSCClientProxy.h b/Engine/Plugins/Runtime/OSC/Source/OSC/Private/OSCClientProxy.h
new file mode 100644
--- /dev/null
+++ b/Engine/Plugins/Runtime/OSC/Source/OSC/Private/OSCClientProxy.h
@@ -0,0 +1,37 @@
+// Copyright 1998-2019 Epic Games, Inc. All Rights Reserved.
+#pragma once
+
+#include "CoreMinimal.h"
+
+#include "Common/UdpSocketReceiver.h"
+#include "Interfaces/IPv4/IPv4Address.h"
+#include "UObject/Object.h"
+
+#include "OSCClient.h"
+
+
+class OSC_API FOSCClientProxy : public IOSCClientProxy
+{
+public:
+	FOSCClientProxy(const FString& InClientName);
+
+	void GetSendIPAddress(FString& InIPAddress, int32& Port) const override;
+	bool SetSendIPAddress(const FString& InIPAddress, const int32 Port) override;
+
+	void SendMessage(FOSCMessage& Message) override;
+	void SendBundle(FOSCBundle& Bundle) override;
+
+	void Stop() override;
+
+private:
+	void SendPacket(IOSCPacket& Packet);
+
+	/** Socket used to send the OSC packets. */
+	FSocket* Socket;
+
+	/** IP Address used by socket. */
+	TSharedPtr<FInternetAddr> IPAddress;
+
+	/** Name of client */
+	FString ClientName;
+};
diff --git a/Engine/Plugins/Runtime/OSC/Source/OSC/Private/OSCManager.cpp b/Engine/Plugins/Runtime/OSC/Source/OSC/Private/OSCManager.cpp
--- a/Engine/Plugins/Runtime/OSC/Source/OSC/Private/OSCManager.cpp
+++ b/Engine/Plugins/Runtime/OSC/Source/OSC/Private/OSCManager.cpp
@@ -57,13 +57,13 @@ namespace
 
 UOSCServer* UOSCManager::CreateOSCServer(FString ReceiveIPAddress, int32 Port, bool bMulticastLoopback, bool bStartListening)
 {
-	UOSCServer* NewOSCServer = NewObject<UOSCServer>();
-
 	if (GetLocalHostAddress(ReceiveIPAddress))
 	{
 		UE_LOG(LogOSC, Display, TEXT("OSCServer ReceiveAddress not specified. Using LocalHost IP: '%s'"), *ReceiveIPAddress);
 	}
 
+	UOSCServer* NewOSCServer = NewObject<UOSCServer>();
+	NewOSCServer->Connect();
 	NewOSCServer->SetMulticastLoopback(bMulticastLoopback);
 	if (NewOSCServer->SetAddress(ReceiveIPAddress, Port))
 	{
@@ -82,13 +82,13 @@ UOSCServer* UOSCManager::CreateOSCServer(FString ReceiveIPAddress, int32 Port, b
 
 UOSCClient* UOSCManager::CreateOSCClient(FString SendIPAddress, int32 Port)
 {
-	UOSCClient* NewOSCClient = NewObject<UOSCClient>();
-
 	if (GetLocalHostAddress(SendIPAddress))
 	{
 		UE_LOG(LogOSC, Display, TEXT("OSCClient SendAddress not specified. Using LocalHost IP: '%s'"), *SendIPAddress);
 	}
 
+	UOSCClient* NewOSCClient = NewObject<UOSCClient>();
+	NewOSCClient->Connect();
 	if (!NewOSCClient->SetSendIPAddress(SendIPAddress, Port))
 	{
 		UE_LOG(LogOSC, Warning, TEXT("Failed to parse SendAddress '%s' for OSCClient. Client unable to send new messages."), *SendIPAddress);
diff --git a/Engine/Plugins/Runtime/OSC/Source/OSC/Private/OSCServer.cpp b/Engine/Plugins/Runtime/OSC/Source/OSC/Private/OSCServer.cpp
--- a/Engine/Plugins/Runtime/OSC/Source/OSC/Private/OSCServer.cpp
+++ b/Engine/Plugins/Runtime/OSC/Source/OSC/Private/OSCServer.cpp
@@ -12,181 +12,55 @@
 #include "OSCBundle.h"
 #include "OSCBundlePacket.h"
 #include "OSCLog.h"
-
-#define OSC_SERVER_BUFFER_FREE 0
-#define OSC_SERVER_BUFFER_LOCKED 1
+#include "OSCServerProxy.h"
 
 
 UOSCServer::UOSCServer(const FObjectInitializer& ObjectInitializer)
-	: Super(ObjectInitializer) 
-	, bWhitelistClients(false)
-	, Socket(nullptr)
-	, SocketReceiver(nullptr)
-	, OSCBufferLock(0)
-	, Port(0)
-	, bMulticastLoopback(false)
+	: Super(ObjectInitializer)
+	, ServerProxy(new FOSCServerProxy(*this))
 {
 }
 
-void UOSCServer::Callback(const FArrayReaderPtr& Data, const FIPv4Endpoint& Endpoint)
+void UOSCServer::Connect()
 {
-	TSharedPtr<IOSCPacket> Packet = IOSCPacket::CreatePacket(Data->GetData());
-	if (Packet.IsValid())
-	{
-		FOSCStream Stream = FOSCStream(Data->GetData(), Data->Num());
-		Packet->ReadData(Stream);
-		OSCPackets.Enqueue(MoveTemp(Packet));
-	}
-
-	if (!FPlatformAtomics::InterlockedCompareExchange(&OSCBufferLock, OSC_SERVER_BUFFER_LOCKED, OSC_SERVER_BUFFER_FREE) &&
-		!OSCPackets.IsEmpty())
-	{
-		check(OSCBufferLock == OSC_SERVER_BUFFER_LOCKED);
-		FSimpleDelegateGraphTask::CreateAndDispatchWhenReady(FSimpleDelegateGraphTask::FDelegate::CreateLambda([this, Endpoint]()
-		{
-			check(OSCBufferLock == OSC_SERVER_BUFFER_LOCKED);
-			FPlatformAtomics::InterlockedCompareExchange(&OSCBufferLock, OSC_SERVER_BUFFER_FREE, OSC_SERVER_BUFFER_LOCKED);
-
-			// Throw request on the ground if endpoint address not whitelisted.
-			if (bWhitelistClients && !ClientWhitelist.Contains(Endpoint.Address.Value))
-			{
-				OSCPackets.Empty();
-				return;
-			}
-
-			TSharedPtr<IOSCPacket> Packet;
-			while (OSCPackets.Dequeue(Packet))
-			{
-				if (Packet->IsMessage())
-				{
-					FOSCMessage Message(FOSCMessage(StaticCastSharedPtr<FOSCMessagePacket>(Packet)));
-					OnOscReceived.Broadcast(Message);
-					UE_LOG(LogOSC, Verbose, TEXT("Message received from endpoint '%s', OSCAddress of '%s'."), *Endpoint.Address.ToString(), *Message.GetAddress().GetFullPath());
-
-					for (const TPair<FOSCAddress, FOSCDispatchMessageEvent>& Pair: AddressPatterns)
-					{
-						const FOSCDispatchMessageEvent& DispatchEvent = Pair.Value;
-						if (Pair.Key.Matches(Message.GetAddress()))
-						{
-							DispatchEvent.Broadcast(Pair.Key, Message);
-							UE_LOG(LogOSC, Verbose, TEXT("Message dispatched from endpoint '%s', OSCAddress path of '%s' matched OSCAddress pattern '%s'."),
-								*Endpoint.Address.ToString(),
-								*Message.GetAddress().GetFullPath(),
-								*Pair.Key.GetFullPath());
-						}
-					}
-				}
-				else if (Packet->IsBundle())
-				{
-					FOSCBundle Bundle(StaticCastSharedPtr<FOSCBundlePacket>(Packet));
-					OnOscBundleReceived.Broadcast(Bundle);
-				}
-				else
-				{
-					UE_LOG(LogOSC, Warning, TEXT("Failed to parse invalid received OSC message. Invalid type (packet is neither identified as message nor bundle)."));
-				}
-			}
-		}),
-		TStatId(),
-		nullptr,
-		ENamedThreads::GameThread);
-	}
+	ServerProxy.Reset(new FOSCServerProxy(*this));
 }
 
 bool UOSCServer::GetMulticastLoopback() const
 {
-	return bMulticastLoopback;
+	check(ServerProxy.IsValid());
+	return ServerProxy->GetMulticastLoopback();
 }
 
 bool UOSCServer::IsActive() const
 {
-	return SocketReceiver != nullptr;
+	check(ServerProxy.IsValid());
+	return ServerProxy->IsActive();
 }
 
 void UOSCServer::Listen()
 {
-	if (IsActive())
-	{
-		UE_LOG(LogOSC, Error, TEXT("OSCServer currently listening: %s:%d. Failed to start new service prior to calling stop."), *GetName(), *ReceiveIPAddress.ToString(), Port);
-		return;
-	}
-
-	FUdpSocketBuilder Builder(*GetName());
-	Builder.BoundToPort(Port);
-	if (ReceiveIPAddress.IsMulticastAddress())
-	{
-		Builder.JoinedToGroup(ReceiveIPAddress);
-		if (bMulticastLoopback)
-		{
-			Builder.WithMulticastLoopback();
-		}
-	}
-	else
-	{
-		if (bMulticastLoopback)
-		{
-			UE_LOG(LogOSC, Warning, TEXT("OSCServer '%s' ReceiveIPAddress provided is not a multicast address.  Not respecting MulticastLoopback boolean."), *GetName());
-		}
-		Builder.BoundToAddress(ReceiveIPAddress);
-	}
-
-	Socket = Builder.Build();
-	if (Socket)
-	{
-		SocketReceiver = new FUdpSocketReceiver(Socket, FTimespan::FromMilliseconds(100), *GetName().Append(TEXT("_ListenerThread")));
-		SocketReceiver->OnDataReceived().BindUObject(this, &UOSCServer::Callback);
-		SocketReceiver->Start();
-
-		UE_LOG(LogOSC, Display, TEXT("OSCServer '%s' Listening: %s:%d."), *GetName(), *ReceiveIPAddress.ToString(), Port);
-	}
-	else
-	{
-		UE_LOG(LogOSC, Warning, TEXT("OSCServer '%s' failed to bind to socket on %s:%d."), *GetName(), *ReceiveIPAddress.ToString(), Port);
-	}
+	check(ServerProxy.IsValid());
+	ServerProxy->Listen(GetName());
 }
 
 bool UOSCServer::SetAddress(const FString& InReceiveIPAddress, int32 InPort)
 {
-	if (IsActive())
-	{
-		UE_LOG(LogOSC, Error, TEXT("Cannot set address while OSCServer is active."));
-		return false;
-	}
-
-	if (!FIPv4Address::Parse(InReceiveIPAddress, ReceiveIPAddress))
-	{
-		UE_LOG(LogOSC, Error, TEXT("Invalid ReceiveIPAddress when attempting to set on OSCServer '%s'. Address not updated."), *GetName());
-		return false;
-	}
-
-	Port = InPort;
-	return true;
+	check(ServerProxy.IsValid());
+	return ServerProxy->SetAddress(InReceiveIPAddress, InPort);
 }
 
 void UOSCServer::SetMulticastLoopback(bool bInMulticastLoopback)
 {
-	if (bInMulticastLoopback != bMulticastLoopback && IsActive())
-	{
-		UE_LOG(LogOSC, Error, TEXT("Cannot update MulticastLoopback while OSCServer '%s' is active."), *GetName());
-		return;
-	}
-
-	bMulticastLoopback = bInMulticastLoopback;
+	check(ServerProxy.IsValid());
+	ServerProxy->SetMulticastLoopback(bInMulticastLoopback);
 }
 
 void UOSCServer::Stop()
 {
-	if (SocketReceiver)
+	if (ServerProxy.IsValid())
 	{
-		delete SocketReceiver;
-		SocketReceiver = nullptr;
-	}
-
-	if (Socket)
-	{
-		Socket->Close();
-		ISocketSubsystem::Get(PLATFORM_SOCKETSUBSYSTEM)->DestroySocket(Socket);
-		Socket = nullptr;
+		ServerProxy->Stop();
 	}
 }
 
@@ -196,75 +70,65 @@ void UOSCServer::BeginDestroy()
 	Super::BeginDestroy();
 }
 
-void UOSCServer::AddWhitelistedClient(const FString& IPAddress)
+void UOSCServer::SetWhitelistClientsEnabled(bool bEnabled)
 {
-	FIPv4Address OutAddress;
-	if (!FIPv4Address::Parse(IPAddress, OutAddress))
-	{
-		UE_LOG(LogOSC, Warning, TEXT("OSCServer '%s' failed to whitelist IP Address '%s'. Address is invalid."), *IPAddress);
-		return;
-	}
-
-	ClientWhitelist.Add(OutAddress.Value);
+	check(ServerProxy.IsValid());
+	ServerProxy->SetWhitelistClientsEnabled(bEnabled);
 }
 
-void UOSCServer::RemoveWhitelistedClient(const FString& IPAddress)
+void UOSCServer::AddWhitelistedClient(const FString& InIPAddress)
 {
-	FIPv4Address OutAddress;
-	if (!FIPv4Address::Parse(IPAddress, OutAddress))
-	{
-		UE_LOG(LogOSC, Warning, TEXT("OSCServer '%s' failed to remove whitelisted IP Address '%s'. Address is invalid."), *IPAddress);
-		return;
-	}
+	check(ServerProxy.IsValid());
+	ServerProxy->AddWhitelistedClient(InIPAddress);
+}
 
-	ClientWhitelist.Remove(OutAddress.Value);
+void UOSCServer::RemoveWhitelistedClient(const FString& InIPAddress)
+{
+	check(ServerProxy.IsValid());
+	ServerProxy->RemoveWhitelistedClient(InIPAddress);
 }
 
 void UOSCServer::ClearWhitelistedClients()
 {
-	ClientWhitelist.Reset();
+	check(ServerProxy.IsValid());
+	ServerProxy->ClearWhitelistedClients();
 }
 
 TSet<FString> UOSCServer::GetWhitelistedClients() const
 {
-	TSet<FString> OutWhitelist;
-	for (uint32 Client : ClientWhitelist)
-	{
-		OutWhitelist.Add(FIPv4Address(Client).ToString());
-	}
-
-	return OutWhitelist;
+	check(ServerProxy.IsValid());
+	return ServerProxy->GetWhitelistedClients();
 }
 
-void UOSCServer::BindEventToOnOSCAddressPatternMatchesPath(const FOSCAddress& OSCAddressPattern, const FOSCDispatchMessageEventBP& Event)
+void UOSCServer::BindEventToOnOSCAddressPatternMatchesPath(const FOSCAddress& InOSCAddressPattern, const FOSCDispatchMessageEventBP& InEvent)
 {
-	if (OSCAddressPattern.IsValidPattern())
+	if (InOSCAddressPattern.IsValidPattern())
 	{
-		FOSCDispatchMessageEvent& MessageEvent = AddressPatterns.FindOrAdd(OSCAddressPattern);
-		MessageEvent.AddUnique(Event);
+		FOSCDispatchMessageEvent& MessageEvent = AddressPatterns.FindOrAdd(InOSCAddressPattern);
+		MessageEvent.AddUnique(InEvent);
 	}
 }
 
-void UOSCServer::UnbindEventFromOnOSCAddressPatternMatchesPath(const FOSCAddress& OSCAddressPattern, const FOSCDispatchMessageEventBP& Event)
+void UOSCServer::UnbindEventFromOnOSCAddressPatternMatchesPath(const FOSCAddress& InOSCAddressPattern, const FOSCDispatchMessageEventBP& InEvent)
 {
-	if (OSCAddressPattern.IsValidPattern())
+	if (InOSCAddressPattern.IsValidPattern())
 	{
-		if (FOSCDispatchMessageEvent* AddressPatternEvent = AddressPatterns.Find(OSCAddressPattern))
+		if (FOSCDispatchMessageEvent* AddressPatternEvent = AddressPatterns.Find(InOSCAddressPattern))
 		{
-			AddressPatternEvent->Remove(Event);
+			AddressPatternEvent->Remove(InEvent);
 			if (!AddressPatternEvent->IsBound())
 			{
-				AddressPatterns.Remove(OSCAddressPattern);
+				AddressPatterns.Remove(InOSCAddressPattern);
 			}
 		}
 	}
 }
 
-void UOSCServer::UnbindAllEventsFromOnOSCAddressPatternMatchesPath(const FOSCAddress& OSCAddressPattern)
+void UOSCServer::UnbindAllEventsFromOnOSCAddressPatternMatchesPath(const FOSCAddress& InOSCAddressPattern)
 {
-	if (OSCAddressPattern.IsValidPattern())
+	if (InOSCAddressPattern.IsValidPattern())
 	{
-		AddressPatterns.Remove(OSCAddressPattern);
+		AddressPatterns.Remove(InOSCAddressPattern);
 	}
 }
 
@@ -282,3 +146,75 @@ TArray<FOSCAddress> UOSCServer::GetBoundOSCAddressPatterns() const
 	}
 	return MoveTemp(OutAddressPatterns);
 }
+
+void UOSCServer::ClearPackets()
+{
+	OSCPackets.Empty();
+}
+
+void UOSCServer::EnqueuePacket(TSharedPtr<IOSCPacket> Packet)
+{
+	OSCPackets.Enqueue(Packet);
+}
+
+void UOSCServer::DispatchBundle(const FString& InIPAddress, const FOSCBundle& InBundle)
+{
+	OnOscBundleReceived.Broadcast(InBundle);
+
+	TSharedPtr<FOSCBundlePacket> BundlePacket = StaticCastSharedPtr<FOSCBundlePacket>(InBundle.GetPacket());
+	FOSCBundlePacket::FPacketBundle Packets = BundlePacket->GetPackets();
+	for (TSharedPtr<IOSCPacket>& Packet : Packets)
+	{
+		if (Packet->IsMessage())
+		{
+			DispatchMessage(InIPAddress, FOSCMessage(Packet));
+		}
+		else if (Packet->IsBundle())
+		{
+			DispatchBundle(InIPAddress, FOSCBundle(Packet));
+		}
+		else
+		{
+			UE_LOG(LogOSC, Warning, TEXT("Failed to parse invalid received message. Invalid OSC type (packet is neither identified as message nor bundle)."));
+		}
+	}
+}
+
+void UOSCServer::DispatchMessage(const FString& InIPAddress, const FOSCMessage& InMessage)
+{
+	OnOscMessageReceived.Broadcast(InMessage);
+	UE_LOG(LogOSC, Verbose, TEXT("Message received from endpoint '%s', OSCAddress of '%s'."), *InIPAddress, *InMessage.GetAddress().GetFullPath());
+
+	for (const TPair<FOSCAddress, FOSCDispatchMessageEvent>& Pair : AddressPatterns)
+	{
+		const FOSCDispatchMessageEvent& DispatchEvent = Pair.Value;
+		if (Pair.Key.Matches(InMessage.GetAddress()))
+		{
+			DispatchEvent.Broadcast(Pair.Key, InMessage);
+			UE_LOG(LogOSC, Verbose, TEXT("Message dispatched from endpoint '%s', OSCAddress path of '%s' matched OSCAddress pattern '%s'."),
+				*InIPAddress,
+				*InMessage.GetAddress().GetFullPath(),
+				*Pair.Key.GetFullPath());
+		}
+	}
+}
+
+void UOSCServer::OnMessageReceived(const FString& InIPAddress)
+{
+	TSharedPtr<IOSCPacket> Packet;
+	while (OSCPackets.Dequeue(Packet))
+	{
+		if (Packet->IsMessage())
+		{
+			DispatchMessage(InIPAddress, FOSCMessage(Packet));
+		}
+		else if (Packet->IsBundle())
+		{
+			DispatchBundle(InIPAddress, FOSCBundle(Packet));
+		}
+		else
+		{
+			UE_LOG(LogOSC, Warning, TEXT("Failed to parse invalid received message. Invalid OSC type (packet is neither identified as message nor bundle)."));
+		}
+	}
+}
\ No newline at end of file
diff --git a/Engine/Plugins/Runtime/OSC/Source/OSC/Private/OSCServerProxy.cpp b/Engine/Plugins/Runtime/OSC/Source/OSC/Private/OSCServerProxy.cpp
new file mode 100644
--- /dev/null
+++ b/Engine/Plugins/Runtime/OSC/Source/OSC/Private/OSCServerProxy.cpp
@@ -0,0 +1,192 @@
+// Copyright 1998-2019 Epic Games, Inc. All Rights Reserved.
+#include "OSCServerProxy.h"
+
+#include "Common/UdpSocketBuilder.h"
+#include "Runtime/Core/Public/Async/TaskGraphInterfaces.h"
+#include "Sockets.h"
+
+#include "OSCLog.h"
+#include "OSCStream.h"
+#include "OSCServer.h"
+
+
+FOSCServerProxy::FOSCServerProxy(UOSCServer& InServer)
+	: Server(&InServer)
+	, Socket(nullptr)
+	, SocketReceiver(nullptr)
+	, Port(0)
+	, bMulticastLoopback(false)
+	, bWhitelistClients(false)
+{
+}
+
+void FOSCServerProxy::OnMessageReceived(const FArrayReaderPtr& Data, const FIPv4Endpoint& Endpoint)
+{
+	TSharedPtr<IOSCPacket> Packet = IOSCPacket::CreatePacket(Data->GetData());
+	if (!Packet.IsValid())
+	{
+		UE_LOG(LogOSC, Verbose, TEXT("Message received from endpoint '%s' invalid OSC packet."), *Endpoint.ToString());
+		return;
+	}
+
+	FOSCStream Stream = FOSCStream(Data->GetData(), Data->Num());
+	Packet->ReadData(Stream);
+	Server->EnqueuePacket(Packet);
+
+	DECLARE_CYCLE_STAT(TEXT("OSCServer.OnMessageReceived"), STAT_OSCServerOnMessageReceived, STATGROUP_OSCNetworkCommands);
+	FFunctionGraphTask::CreateAndDispatchWhenReady([this, Endpoint]()
+	{
+		// Throw request on the ground if endpoint address not whitelisted.
+		if (bWhitelistClients && !ClientWhitelist.Contains(Endpoint.Address.Value))
+		{
+			Server->ClearPackets();
+			return;
+		}
+
+		Server->OnMessageReceived(Endpoint.Address.ToString());
+	}, GET_STATID(STAT_OSCServerOnMessageReceived), nullptr, ENamedThreads::GameThread);
+}
+
+bool FOSCServerProxy::GetMulticastLoopback() const
+{
+	return bMulticastLoopback;
+}
+
+bool FOSCServerProxy::IsActive() const
+{
+	return SocketReceiver != nullptr;
+}
+
+void FOSCServerProxy::Listen(const FString& ServerName)
+{
+	if (IsActive())
+	{
+		UE_LOG(LogOSC, Error, TEXT("OSCServer currently listening: %s:%d. Failed to start new service prior to calling stop."),
+			*ServerName, *ReceiveIPAddress.ToString(), Port);
+		return;
+	}
+
+	FUdpSocketBuilder Builder(*ServerName);
+	Builder.BoundToPort(Port);
+	if (ReceiveIPAddress.IsMulticastAddress())
+	{
+		Builder.JoinedToGroup(ReceiveIPAddress);
+		if (bMulticastLoopback)
+		{
+			Builder.WithMulticastLoopback();
+		}
+	}
+	else
+	{
+		if (bMulticastLoopback)
+		{
+			UE_LOG(LogOSC, Warning, TEXT("OSCServer '%s' ReceiveIPAddress provided is not a multicast address.  Not respecting MulticastLoopback boolean."),
+				*ServerName);
+		}
+		Builder.BoundToAddress(ReceiveIPAddress);
+	}
+
+	Socket = Builder.Build();
+	if (Socket)
+	{
+		SocketReceiver = new FUdpSocketReceiver(Socket, FTimespan::FromMilliseconds(100), *(ServerName + TEXT("_ListenerThread")));
+		SocketReceiver->OnDataReceived().BindRaw(this, &FOSCServerProxy::OnMessageReceived);
+		SocketReceiver->Start();
+
+		UE_LOG(LogOSC, Display, TEXT("OSCServer '%s' Listening: %s:%d."), *ServerName, *ReceiveIPAddress.ToString(), Port);
+	}
+	else
+	{
+		UE_LOG(LogOSC, Warning, TEXT("OSCServer '%s' failed to bind to socket on %s:%d."), *ServerName, *ReceiveIPAddress.ToString(), Port);
+	}
+}
+
+bool FOSCServerProxy::SetAddress(const FString& InReceiveIPAddress, int32 InPort)
+{
+	if (IsActive())
+	{
+		UE_LOG(LogOSC, Error, TEXT("Cannot set address while OSCServer is active."));
+		return false;
+	}
+
+	if (!FIPv4Address::Parse(InReceiveIPAddress, ReceiveIPAddress))
+	{
+		UE_LOG(LogOSC, Error, TEXT("Invalid ReceiveIPAddress '%s'. OSCServer ReceiveIP Address not updated."), *InReceiveIPAddress);
+		return false;
+	}
+
+	Port = InPort;
+	return true;
+}
+
+void FOSCServerProxy::SetMulticastLoopback(bool bInMulticastLoopback)
+{
+	if (bInMulticastLoopback != bMulticastLoopback && IsActive())
+	{
+		UE_LOG(LogOSC, Error, TEXT("Cannot update MulticastLoopback while OSCServer is active."));
+		return;
+	}
+
+	bMulticastLoopback = bInMulticastLoopback;
+}
+
+void FOSCServerProxy::Stop()
+{
+	if (SocketReceiver)
+	{
+		delete SocketReceiver;
+		SocketReceiver = nullptr;
+	}
+
+	if (Socket)
+	{
+		Socket->Close();
+		ISocketSubsystem::Get(PLATFORM_SOCKETSUBSYSTEM)->DestroySocket(Socket);
+		Socket = nullptr;
+	}
+}
+
+void FOSCServerProxy::AddWhitelistedClient(const FString& IPAddress)
+{
+	FIPv4Address OutAddress;
+	if (!FIPv4Address::Parse(IPAddress, OutAddress))
+	{
+		UE_LOG(LogOSC, Warning, TEXT("OSCServer '%s' failed to whitelist IP Address '%s'. Address is invalid."), *IPAddress);
+		return;
+	}
+
+	ClientWhitelist.Add(OutAddress.Value);
+}
+
+void FOSCServerProxy::RemoveWhitelistedClient(const FString& IPAddress)
+{
+	FIPv4Address OutAddress;
+	if (!FIPv4Address::Parse(IPAddress, OutAddress))
+	{
+		UE_LOG(LogOSC, Warning, TEXT("OSCServer '%s' failed to remove whitelisted IP Address '%s'. Address is invalid."), *IPAddress);
+		return;
+	}
+
+	ClientWhitelist.Remove(OutAddress.Value);
+}
+
+void FOSCServerProxy::ClearWhitelistedClients()
+{
+	ClientWhitelist.Reset();
+}
+
+TSet<FString> FOSCServerProxy::GetWhitelistedClients() const
+{
+	TSet<FString> OutWhitelist;
+	for (uint32 Client : ClientWhitelist)
+	{
+		OutWhitelist.Add(FIPv4Address(Client).ToString());
+	}
+
+	return OutWhitelist;
+}
+
+void FOSCServerProxy::SetWhitelistClientsEnabled(bool bEnabled)
+{
+	bWhitelistClients = bEnabled;
+}
diff --git a/Engine/Plugins/Runtime/OSC/Source/OSC/Private/OSCServerProxy.h b/Engine/Plugins/Runtime/OSC/Source/OSC/Private/OSCServerProxy.h
new file mode 100644
--- /dev/null
+++ b/Engine/Plugins/Runtime/OSC/Source/OSC/Private/OSCServerProxy.h
@@ -0,0 +1,63 @@
+// Copyright 1998-2019 Epic Games, Inc. All Rights Reserved.
+#pragma once
+
+#include "CoreMinimal.h"
+
+#include "Common/UdpSocketReceiver.h"
+#include "Interfaces/IPv4/IPv4Address.h"
+#include "Interfaces/IPv4/IPv4Endpoint.h"
+#include "UObject/Object.h"
+
+#include "OSCServer.h"
+
+
+class OSC_API FOSCServerProxy : public IOSCServerProxy
+{
+public:
+	FOSCServerProxy(UOSCServer& InServer);
+	virtual ~FOSCServerProxy() { }
+
+	bool GetMulticastLoopback() const override;
+
+	bool IsActive() const override;
+
+	void Listen(const FString& ServerName) override;
+
+	bool SetAddress(const FString& ReceiveIPAddress, int32 Port) override;
+	void SetMulticastLoopback(bool bMulticastLoopback) override;
+
+	void Stop() override;
+
+	void AddWhitelistedClient(const FString& IPAddress) override;
+	void RemoveWhitelistedClient(const FString& IPAddress) override;
+	void ClearWhitelistedClients() override;
+	TSet<FString> GetWhitelistedClients() const override;
+
+	void SetWhitelistClientsEnabled(bool bEnabled) override;
+
+	/** Callback that receives data from a socket. */
+	void OnMessageReceived(const FArrayReaderPtr& Data, const FIPv4Endpoint& Endpoint);
+
+	UOSCServer* Server;
+
+	/** Socket used to listen for OSC packets. */
+	FSocket* Socket;
+
+	/** UDP receiver. */
+	FUdpSocketReceiver* SocketReceiver;
+
+	/** Set of client addresses whitelisted to process packets from. */
+	TSet<uint32> ClientWhitelist;
+
+	/** IPAddress to listen for OSC packets on.  If unset, defaults to LocalHost */
+	FIPv4Address ReceiveIPAddress;
+
+	/** Port to listen for OSC packets on. */
+	int32 Port;
+
+	/** Whether or not to loopback if address provided is multicast */
+	bool bMulticastLoopback;
+
+	/** Whether or not to use client whitelist */
+	bool bWhitelistClients;
+};
diff --git a/Engine/Plugins/Runtime/OSC/Source/OSC/Public/OSCClient.h b/Engine/Plugins/Runtime/OSC/Source/OSC/Public/OSCClient.h
--- a/Engine/Plugins/Runtime/OSC/Source/OSC/Public/OSCClient.h
+++ b/Engine/Plugins/Runtime/OSC/Source/OSC/Public/OSCClient.h
@@ -2,14 +2,29 @@
 #pragma once
 
 #include "CoreMinimal.h"
+#include "Templates/UniquePtr.h"
 #include "UObject/Object.h"
-#include "Interfaces/IPv4/IPv4Address.h"
-#include "IPAddress.h"
+
 #include "OSCMessage.h"
 #include "OSCBundle.h"
+
 #include "OSCClient.generated.h"
 
-class FSocket;
+
+/** Interface for internal networking implementation.  See UOSCClient for details */
+class OSC_API IOSCClientProxy
+{
+public:
+	virtual ~IOSCClientProxy() { }
+
+	virtual void GetSendIPAddress(FString& InIPAddress, int32& Port) const = 0;
+	virtual bool SetSendIPAddress(const FString& InIPAddress, const int32 Port) = 0;
+
+	virtual void SendMessage(FOSCMessage& Message) = 0;
+	virtual void SendBundle(FOSCBundle& Bundle) = 0;
+
+	virtual void Stop() = 0;
+};
 
 UCLASS(BlueprintType)
 class OSC_API UOSCClient : public UObject
@@ -17,17 +32,17 @@ class OSC_API UOSCClient : public UObject
 	GENERATED_UCLASS_BODY()
 
 public:
-	UOSCClient();
-
-	/** Sets the OSC Client IP address and port. Returns whether 
-	  * address and port was successfully set. */
-	UFUNCTION(BlueprintCallable, Category = "Audio|OSC")
-	bool SetSendIPAddress(const FString& IPAddress, const int32 Port);
+	void Connect();
 
 	/** Gets the OSC Client IP address and port. */
 	UFUNCTION(BlueprintCallable, Category = "Audio|OSC")
 	void GetSendIPAddress(UPARAM(ref) FString& IPAddress, UPARAM(ref) int32& Port);
 
+	/** Sets the OSC Client IP address and port. Returns whether
+	  * address and port was successfully set. */
+	UFUNCTION(BlueprintCallable, Category = "Audio|OSC")
+	bool SetSendIPAddress(const FString& IPAddress, const int32 Port);
+
 	/** Send OSC message to  a specific address. */
 	UFUNCTION(BlueprintCallable, Category = "Audio|OSC")
 	void SendOSCMessage(UPARAM(ref) FOSCMessage& Message);
@@ -37,18 +52,11 @@ public:
 	void SendOSCBundle(UPARAM(ref) FOSCBundle& Bundle);
 
 protected:
-
 	void BeginDestroy() override;
-	
+
 	/** Stop and tidy up network socket. */
 	void Stop();
 	
-	/** Send OSC packet data. */
-	void SendPacket(IOSCPacket& InPacket);
-
-	/** Socket used to send the OSC packets. */
-	FSocket* Socket;
-
-	/** IP Address used by socket. */
-	TSharedPtr<FInternetAddr> IPAddress;
+	/** Pointer to internal implementation of client proxy */
+	TUniquePtr<IOSCClientProxy> ClientProxy;
 };
diff --git a/Engine/Plugins/Runtime/OSC/Source/OSC/Public/OSCModule.cpp b/Engine/Plugins/Runtime/OSC/Source/OSC/Public/OSCModule.cpp
--- a/Engine/Plugins/Runtime/OSC/Source/OSC/Public/OSCModule.cpp
+++ b/Engine/Plugins/Runtime/OSC/Source/OSC/Public/OSCModule.cpp
@@ -10,14 +10,6 @@ DEFINE_LOG_CATEGORY(LogOSC);
 
 class FOSCModule : public IModuleInterface
 {
-public:
-	virtual void StartupModule()
-	{
-		if (!FModuleManager::Get().LoadModule(TEXT("Networking")))
-		{
-			UE_LOG(LogOSC, Error, TEXT("Required module 'Networking' failed to load. OSC service disabled."));
-		}
-	}
 };
 
 IMPLEMENT_MODULE(FOSCModule, OSC)
diff --git a/Engine/Plugins/Runtime/OSC/Source/OSC/Public/OSCServer.h b/Engine/Plugins/Runtime/OSC/Source/OSC/Public/OSCServer.h
--- a/Engine/Plugins/Runtime/OSC/Source/OSC/Public/OSCServer.h
+++ b/Engine/Plugins/Runtime/OSC/Source/OSC/Public/OSCServer.h
@@ -3,26 +3,48 @@
 
 #include "CoreMinimal.h"
 
-#include "Common/UdpSocketReceiver.h"
+#include "Async/TaskGraphInterfaces.h"
 #include "Containers/Queue.h"
-#include "Interfaces/IPv4/IPv4Address.h"
-#include "Interfaces/IPv4/IPv4Endpoint.h"
-#include "OSCMessage.h"
+#include "Templates/UniquePtr.h"
+#include "UObject/Object.h"
+
 #include "OSCBundle.h"
+#include "OSCMessage.h"
 #include "OSCPacket.h"
-#include "UObject/Object.h"
 
 #include "OSCServer.generated.h"
 
 // Forward Declarations
 class FSocket;
 
+
+// Delegates
 DECLARE_DYNAMIC_MULTICAST_DELEGATE_OneParam(FOSCReceivedMessageEvent, const FOSCMessage&, Message);
 DECLARE_DYNAMIC_MULTICAST_DELEGATE_TwoParams(FOSCDispatchMessageEvent, const FOSCAddress&, AddressPattern, const FOSCMessage&, Message);
 DECLARE_DYNAMIC_MULTICAST_DELEGATE_OneParam(FOSCReceivedBundleEvent, const FOSCBundle&, Bundle);
-
 DECLARE_DYNAMIC_DELEGATE_TwoParams(FOSCDispatchMessageEventBP, const FOSCAddress&, AddressPattern, const FOSCMessage&, Message);
 
+DECLARE_STATS_GROUP(TEXT("OSC Commands"), STATGROUP_OSCNetworkCommands, STATCAT_Advanced);
+
+
+/** Interface for internal networking implementation.  See UOSCServer for details */
+class OSC_API IOSCServerProxy
+{
+public:
+	virtual ~IOSCServerProxy() { }
+
+	virtual bool GetMulticastLoopback() const = 0;
+	virtual bool IsActive() const = 0;
+	virtual void Listen(const FString& ServerName) = 0;
+	virtual bool SetAddress(const FString& InReceiveIPAddress, int32 InPort) = 0;
+	virtual void SetMulticastLoopback(bool bInMulticastLoopback) = 0;
+	virtual void Stop() = 0;
+	virtual void AddWhitelistedClient(const FString& InIPAddress) = 0;
+	virtual void RemoveWhitelistedClient(const FString& IPAddress) = 0;
+	virtual void ClearWhitelistedClients() = 0;
+	virtual TSet<FString> GetWhitelistedClients() const = 0;
+	virtual void SetWhitelistClientsEnabled(bool bEnabled) = 0;
+};
 
 UCLASS(BlueprintType)
 class OSC_API UOSCServer : public UObject
@@ -30,6 +52,8 @@ class OSC_API UOSCServer : public UObject
 	GENERATED_UCLASS_BODY()
 
 public:
+	void Connect();
+
 	/** Gets whether or not to loopback if ReceiveIPAddress provided is multicast. */
 	UFUNCTION(BlueprintCallable, Category = "Audio|OSC")
 	bool GetMulticastLoopback() const;
@@ -56,15 +80,17 @@ public:
 
 	/** Event that gets called when an OSC message is received. */
 	UPROPERTY(BlueprintAssignable, Category = "Audio|OSC")
-	FOSCReceivedMessageEvent OnOscReceived;
+	FOSCReceivedMessageEvent OnOscMessageReceived;
 
 	/** Event that gets called when an OSC bundle is received. */
 	UPROPERTY(BlueprintAssignable, Category = "Audio|OSC")
 	FOSCReceivedBundleEvent OnOscBundleReceived;
-	
-	/** When true, only process received events from whitelisted clients. */
-	UPROPERTY(BlueprintReadWrite, Category = "Audio|OSC")
-	bool bWhitelistClients;
+
+	/** When set to true, server will only process received 
+	  * messages from whitelisted clients.
+	  */
+	UFUNCTION(BlueprintCallable, Category = "Audio|OSC")
+	void SetWhitelistClientsEnabled(bool bEnabled);
 
 	/** Adds client to whitelist of clients to listen for. */
 	UFUNCTION(BlueprintCallable, Category = "Audio|OSC")
@@ -102,36 +128,31 @@ public:
 	UFUNCTION(BlueprintCallable, Category = "Audio|OSC")
 	TArray<FOSCAddress> GetBoundOSCAddressPatterns() const;
 
-protected:
-	void BeginDestroy() override;
-	
-	/** Callback that receives data from a socket. */
-	void Callback(const FArrayReaderPtr& Data, const FIPv4Endpoint& Endpoint);
-	
-	/** Queue stores incoming OSC packet requests to process on the game thread. */
-	TQueue<TSharedPtr<IOSCPacket>> OSCPackets;
+	/** Clears all packets pending processing */
+	void ClearPackets();
+
+	/** Enqueues packet to be processed */
+	void EnqueuePacket(TSharedPtr<IOSCPacket>);
 
-	/** Socket used to listen for OSC packets. */
-	FSocket* Socket;
+	/** Callback for when message is received by server */
+	void OnMessageReceived(const FString& InIPAddress);
 
-	/** UDP receiver. */
-	FUdpSocketReceiver* SocketReceiver;
+protected:
+	void BeginDestroy() override;
 
-	/** Set of client addresses whitelisted to process packets from. */
-	TSet<uint32> ClientWhitelist;
+private:
+	/** Dispatches provided bundle received */
+	void DispatchBundle(const FString& InIPAddress, const FOSCBundle& InBundle);
 
-	/** OSC buffer lock. */
-	int32 OSCBufferLock;
+	/** Dispatches provided message received */
+	void DispatchMessage(const FString& InIPAddress, const FOSCMessage& InMessage);
 
-	/** IPAddress to listen for OSC packets on.  If unset, defaults to LocalHost */
-	FIPv4Address ReceiveIPAddress;
+	/** Pointer to internal implementation of server proxy */
+	TUniquePtr<IOSCServerProxy> ServerProxy;
 
-	/** Port to listen for OSC packets on. */
-	int32 Port;
+	/** Queue stores incoming OSC packet requests to process on the game thread. */
+	TQueue<TSharedPtr<IOSCPacket>> OSCPackets;
 
-	/** Whether or not to loopback if address provided is multicast */
-	bool bMulticastLoopback;
-	
-	/** Address pattern hash to  to check against when dispatching incoming messages */
+	/** Address pattern hash to check against when dispatching incoming messages */
 	TMap<FOSCAddress, FOSCDispatchMessageEvent> AddressPatterns;
 };
