Add back code in module start-up to ensure networking is loaded for packaged builds.

diff --git a/Engine/Plugins/Runtime/OSC/Source/OSC/Public/OSCModule.cpp b/Engine/Plugins/Runtime/OSC/Source/OSC/Public/OSCModule.cpp
--- a/Engine/Plugins/Runtime/OSC/Source/OSC/Public/OSCModule.cpp
+++ b/Engine/Plugins/Runtime/OSC/Source/OSC/Public/OSCModule.cpp
@@ -10,6 +10,14 @@ DEFINE_LOG_CATEGORY(LogOSC);
 
 class FOSCModule : public IModuleInterface
 {
+public:
+	virtual void StartupModule()
+	{
+		if (!FModuleManager::Get().LoadModule(TEXT("Networking")))
+		{
+			UE_LOG(LogOSC, Error, TEXT("Required module 'Networking' failed to load. OSC service disabled."));
+		}
+	}
 };
 
 IMPLEMENT_MODULE(FOSCModule, OSC)
