Fix bug with reading OSC string sometimes not finding null terminator introduced in earlier CL.

diff --git a/Engine/Plugins/Runtime/OSC/Source/OSC/Private/OSCStream.cpp b/Engine/Plugins/Runtime/OSC/Source/OSC/Private/OSCStream.cpp
--- a/Engine/Plugins/Runtime/OSC/Source/OSC/Private/OSCStream.cpp
+++ b/Engine/Plugins/Runtime/OSC/Source/OSC/Private/OSCStream.cpp
@@ -351,11 +351,9 @@ FString FOSCStream::ReadString()
 
 	const ANSICHAR* StrStart = (ANSICHAR*)(&Data[InitPosition]);
 
-	int32 i = InitPosition;
-	for (; i < DataSize; i++)
+	for (; Position < DataSize; Position++)
 	{
-		Position++;
-		if (Data[i] == '\0')
+		if (Data[Position] == '\0')
 		{
 			break;
 		}
