OSC string parse pad position required increment to next character.

diff --git a/Engine/Plugins/Runtime/OSC/Source/OSC/Private/OSCStream.cpp b/Engine/Plugins/Runtime/OSC/Source/OSC/Private/OSCStream.cpp
--- a/Engine/Plugins/Runtime/OSC/Source/OSC/Private/OSCStream.cpp
+++ b/Engine/Plugins/Runtime/OSC/Source/OSC/Private/OSCStream.cpp
@@ -365,8 +365,10 @@ FString FOSCStream::ReadString()
 		return FString();
 	}
 
-	// Pad position before returning
+	// Note end for string copy, increment to next read
+	// location, and pad position.
 	const int32 EndPosition = Position;
+	Position++;
 	Position = ((Position + 3) / 4) * 4;
 
 	return FString(EndPosition - InitPosition, StrStart);
