- Fix bug with OSC bundle corruption due to read not accounting for header element offset
- Clarify server's function call (not necessarily OnMessageReceive, but OnPacketReceive and then diverts to packet type accordingly)

diff --git a/Engine/Plugins/Runtime/OSC/Source/OSC/Private/OSCBundlePacket.cpp b/Engine/Plugins/Runtime/OSC/Source/OSC/Private/OSCBundlePacket.cpp
--- a/Engine/Plugins/Runtime/OSC/Source/OSC/Private/OSCBundlePacket.cpp
+++ b/Engine/Plugins/Runtime/OSC/Source/OSC/Private/OSCBundlePacket.cpp
@@ -50,15 +50,34 @@ void FOSCBundlePacket::WriteData(FOSCStream& Stream)
 
 void FOSCBundlePacket::ReadData(FOSCStream& Stream)
 {
+	FString BundleTag = Stream.ReadString();
+	if (BundleTag != OSC::BundleTag)
+	{
+		UE_LOG(LogOSC, Warning, TEXT("Failed to parse OSCBundle of invalid format. #bundle identifier not first item in packet."));
+		return;
+	}
+
 	TimeTag = FOSCType(Stream.ReadUInt64());
 
 	while (!Stream.HasReachedEnd())
 	{
-		TSharedPtr<IOSCPacket> Packet = IOSCPacket::CreatePacket(Stream.GetData());
-		if (Packet.IsValid())
+		int32 PacketLength = Stream.ReadInt32();
+
+		int32 StartPos = Stream.GetPosition();
+		TSharedPtr<IOSCPacket> Packet = IOSCPacket::CreatePacket(Stream.GetData() + Stream.GetPosition());
+		if (!Packet.IsValid())
+		{
+			break;
+		}
+
+		Packet->ReadData(Stream);
+		Packets.Add(Packet);
+		int32 EndPos = Stream.GetPosition();
+
+		if (EndPos - StartPos != PacketLength)
 		{
-			Packet->ReadData(Stream);
-			Packets.Add(Packet);
+			UE_LOG(LogOSC, Warning, TEXT("Failed to parse OSCBundle of invalid format. Element size mismatch."));
+			break;
 		}
 	}
 }
diff --git a/Engine/Plugins/Runtime/OSC/Source/OSC/Private/OSCPacket.cpp b/Engine/Plugins/Runtime/OSC/Source/OSC/Private/OSCPacket.cpp
--- a/Engine/Plugins/Runtime/OSC/Source/OSC/Private/OSCPacket.cpp
+++ b/Engine/Plugins/Runtime/OSC/Source/OSC/Private/OSCPacket.cpp
@@ -27,7 +27,7 @@ TSharedPtr<IOSCPacket> IOSCPacket::CreatePacket(const uint8* PacketType)
 	}
 	else
 	{
-		UE_LOG(LogOSC, Warning, TEXT("Failed to parse lead character of OSC message packet. "
+		UE_LOG(LogOSC, Warning, TEXT("Failed to parse lead character of OSC packet. "
 			"Lead identifier of '%c' not valid bundle tag ('%s') or message ('%s') identifier."), *PacketIdentifier, *OSC::BundleTag, *OSC::PathSeparator);
 		return nullptr;
 	}
diff --git a/Engine/Plugins/Runtime/OSC/Source/OSC/Private/OSCServer.cpp b/Engine/Plugins/Runtime/OSC/Source/OSC/Private/OSCServer.cpp
--- a/Engine/Plugins/Runtime/OSC/Source/OSC/Private/OSCServer.cpp
+++ b/Engine/Plugins/Runtime/OSC/Source/OSC/Private/OSCServer.cpp
@@ -199,7 +199,7 @@ void UOSCServer::DispatchMessage(const FString& InIPAddress, const FOSCMessage&
 	}
 }
 
-void UOSCServer::OnMessageReceived(const FString& InIPAddress)
+void UOSCServer::OnPacketReceived(const FString& InIPAddress)
 {
 	TSharedPtr<IOSCPacket> Packet;
 	while (OSCPackets.Dequeue(Packet))
diff --git a/Engine/Plugins/Runtime/OSC/Source/OSC/Private/OSCServerProxy.cpp b/Engine/Plugins/Runtime/OSC/Source/OSC/Private/OSCServerProxy.cpp
--- a/Engine/Plugins/Runtime/OSC/Source/OSC/Private/OSCServerProxy.cpp
+++ b/Engine/Plugins/Runtime/OSC/Source/OSC/Private/OSCServerProxy.cpp
@@ -20,7 +20,7 @@ FOSCServerProxy::FOSCServerProxy(UOSCServer& InServer)
 {
 }
 
-void FOSCServerProxy::OnMessageReceived(const FArrayReaderPtr& Data, const FIPv4Endpoint& Endpoint)
+void FOSCServerProxy::OnPacketReceived(const FArrayReaderPtr& Data, const FIPv4Endpoint& Endpoint)
 {
 	TSharedPtr<IOSCPacket> Packet = IOSCPacket::CreatePacket(Data->GetData());
 	if (!Packet.IsValid())
@@ -33,7 +33,7 @@ void FOSCServerProxy::OnMessageReceived(const FArrayReaderPtr& Data, const FIPv4
 	Packet->ReadData(Stream);
 	Server->EnqueuePacket(Packet);
 
-	DECLARE_CYCLE_STAT(TEXT("OSCServer.OnMessageReceived"), STAT_OSCServerOnMessageReceived, STATGROUP_OSCNetworkCommands);
+	DECLARE_CYCLE_STAT(TEXT("OSCServer.OnPacketReceived"), STAT_OSCServerOnPacketReceived, STATGROUP_OSCNetworkCommands);
 	FFunctionGraphTask::CreateAndDispatchWhenReady([this, Endpoint]()
 	{
 		// Throw request on the ground if endpoint address not whitelisted.
@@ -43,8 +43,8 @@ void FOSCServerProxy::OnMessageReceived(const FArrayReaderPtr& Data, const FIPv4
 			return;
 		}
 
-		Server->OnMessageReceived(Endpoint.Address.ToString());
-	}, GET_STATID(STAT_OSCServerOnMessageReceived), nullptr, ENamedThreads::GameThread);
+		Server->OnPacketReceived(Endpoint.Address.ToString());
+	}, GET_STATID(STAT_OSCServerOnPacketReceived), nullptr, ENamedThreads::GameThread);
 }
 
 bool FOSCServerProxy::GetMulticastLoopback() const
@@ -90,7 +90,7 @@ void FOSCServerProxy::Listen(const FString& ServerName)
 	if (Socket)
 	{
 		SocketReceiver = new FUdpSocketReceiver(Socket, FTimespan::FromMilliseconds(100), *(ServerName + TEXT("_ListenerThread")));
-		SocketReceiver->OnDataReceived().BindRaw(this, &FOSCServerProxy::OnMessageReceived);
+		SocketReceiver->OnDataReceived().BindRaw(this, &FOSCServerProxy::OnPacketReceived);
 		SocketReceiver->Start();
 
 		UE_LOG(LogOSC, Display, TEXT("OSCServer '%s' Listening: %s:%d."), *ServerName, *ReceiveIPAddress.ToString(), Port);
diff --git a/Engine/Plugins/Runtime/OSC/Source/OSC/Private/OSCServerProxy.h b/Engine/Plugins/Runtime/OSC/Source/OSC/Private/OSCServerProxy.h
--- a/Engine/Plugins/Runtime/OSC/Source/OSC/Private/OSCServerProxy.h
+++ b/Engine/Plugins/Runtime/OSC/Source/OSC/Private/OSCServerProxy.h
@@ -36,7 +36,7 @@ public:
 	void SetWhitelistClientsEnabled(bool bEnabled) override;
 
 	/** Callback that receives data from a socket. */
-	void OnMessageReceived(const FArrayReaderPtr& Data, const FIPv4Endpoint& Endpoint);
+	void OnPacketReceived(const FArrayReaderPtr& Data, const FIPv4Endpoint& Endpoint);
 
 	UOSCServer* Server;
 
diff --git a/Engine/Plugins/Runtime/OSC/Source/OSC/Public/OSCServer.h b/Engine/Plugins/Runtime/OSC/Source/OSC/Public/OSCServer.h
--- a/Engine/Plugins/Runtime/OSC/Source/OSC/Public/OSCServer.h
+++ b/Engine/Plugins/Runtime/OSC/Source/OSC/Public/OSCServer.h
@@ -134,8 +134,8 @@ public:
 	/** Enqueues packet to be processed */
 	void EnqueuePacket(TSharedPtr<IOSCPacket>);
 
-	/** Callback for when message is received by server */
-	void OnMessageReceived(const FString& InIPAddress);
+	/** Callback for when packet is received by server */
+	void OnPacketReceived(const FString& InIPAddress);
 
 protected:
 	void BeginDestroy() override;
