Fix comment typo

diff --git a/Engine/Plugins/Runtime/OSC/Source/OSC/Public/OSCServer.h b/Engine/Plugins/Runtime/OSC/Source/OSC/Public/OSCServer.h
--- a/Engine/Plugins/Runtime/OSC/Source/OSC/Public/OSCServer.h
+++ b/Engine/Plugins/Runtime/OSC/Source/OSC/Public/OSCServer.h
@@ -112,7 +112,7 @@ public:
 	UFUNCTION(BlueprintCallable, Category = "Audio|OSC")
 	void BindEventToOnOSCAddressPatternMatchesPath(const FOSCAddress& OSCAddressPattern, const FOSCDispatchMessageEventBP& Event);
 
-	/** Unbinds specific event from OSCAddress patter */
+	/** Unbinds specific event from OSCAddress pattern. */
 	UFUNCTION(BlueprintCallable, Category = "Audio|OSC")
 	void UnbindEventFromOnOSCAddressPatternMatchesPath(const FOSCAddress& OSCAddressPattern, const FOSCDispatchMessageEventBP& Event);
 
