- Clean Up OSC address function comments in blueprint
- Add optimized Push/PopContainers (multiple) to FOSCAddress (only caches aggregates after removing/adding all)
- Add missing BP call for FOSCAddress::ClearContainers

diff --git a/Engine/Plugins/Runtime/OSC/Source/OSC/Private/OSCAddress.cpp b/Engine/Plugins/Runtime/OSC/Source/OSC/Private/OSCAddress.cpp
--- a/Engine/Plugins/Runtime/OSC/Source/OSC/Private/OSCAddress.cpp
+++ b/Engine/Plugins/Runtime/OSC/Source/OSC/Private/OSCAddress.cpp
@@ -81,7 +81,6 @@ bool FOSCAddress::IsValidPath() const
 
 void FOSCAddress::PushContainer(const FString& Container)
 {
-	const bool bInvalidateSeparator = true;
 	if (Container.Contains(OSC::PathSeparator))
 	{
 		UE_LOG(LogOSC, Warning, TEXT("Failed to push container on OSCAddress. "
@@ -93,18 +92,62 @@ void FOSCAddress::PushContainer(const FString& Container)
 	CacheAggregates();
 }
 
+void FOSCAddress::PushContainers(const TArray<FString>& InContainers)
+{
+	for (const FString& Container : InContainers)
+	{
+		if (Container.Contains(OSC::PathSeparator))
+		{
+			UE_LOG(LogOSC, Warning, TEXT("Failed to push containers on OSCAddress. "
+				"Cannot contain OSC path separator '%s'."), *OSC::PathSeparator);
+			return;
+		}
+	}
+
+	for (const FString& Container : InContainers)
+	{
+		Containers.Push(Container);
+	}
+
+	CacheAggregates();
+}
+
 FString FOSCAddress::PopContainer()
 {
+	FString Popped;
 	if (Containers.Num() > 0)
 	{
-		FString Popped = Containers.Pop(false);
+		Popped = Containers.Pop(false /* bAllowShrinking */);
 		Hash = GetTypeHash(GetFullPath());
 	}
 
-	return FString();
+	return Popped;
+}
+
+TArray<FString> FOSCAddress::PopContainers(int32 InNumContainers)
+{
+	TArray<FString> Popped;
+	if (InNumContainers <= 0 || Containers.Num() == 0)
+	{
+		return Popped;
+	}
+
+	int32 Removed = 0;
+	for (int32 i = Containers.Num() - 1; i >= 0; --i)
+	{
+		if (Removed > InNumContainers)
+		{
+			break;
+		}
+		++Removed;
+		Popped.Add(Containers.Pop(false /* bAllowShrinking */));
+	}
+
+	Hash = GetTypeHash(GetFullPath());
+	return Popped;
 }
 
-void FOSCAddress::ClearContainers(const FString& Container)
+void FOSCAddress::ClearContainers()
 {
 	Containers.Reset();
 	CacheAggregates();
diff --git a/Engine/Plugins/Runtime/OSC/Source/OSC/Private/OSCManager.cpp b/Engine/Plugins/Runtime/OSC/Source/OSC/Private/OSCManager.cpp
--- a/Engine/Plugins/Runtime/OSC/Source/OSC/Private/OSCManager.cpp
+++ b/Engine/Plugins/Runtime/OSC/Source/OSC/Private/OSCManager.cpp
@@ -187,8 +187,28 @@ FOSCMessage& UOSCManager::AddBool(FOSCMessage& Message, bool Value)
 	return Message;
 }
 
-void UOSCManager::GetMessagesFromBundle(const FOSCBundle& Bundle, TArray<FOSCMessage>& Messages)
+TArray<FOSCBundle>& UOSCManager::GetBundlesFromBundle(const FOSCBundle& Bundle, TArray<FOSCBundle>& Bundles)
 {
+	Bundles.Reset();
+	if (Bundle.GetPacket().IsValid())
+	{
+		const TSharedPtr<FOSCBundlePacket>& BundlePacket = StaticCastSharedPtr<FOSCBundlePacket>(Bundle.GetPacket());
+		for (int32 i = 0; i < BundlePacket->GetPackets().Num(); i++)
+		{
+			const TSharedPtr<IOSCPacket>& Packet = BundlePacket->GetPackets()[i];
+			if (Packet->IsBundle())
+			{
+				Bundles.Emplace(StaticCastSharedPtr<FOSCMessagePacket>(Packet));
+			}
+		}
+	}
+
+	return Bundles;
+}
+
+TArray<FOSCMessage>& UOSCManager::GetMessagesFromBundle(const FOSCBundle& Bundle, TArray<FOSCMessage>& Messages)
+{
+	Messages.Reset();
 	if (Bundle.GetPacket().IsValid())
 	{
 		const TSharedPtr<FOSCBundlePacket>& BundlePacket = StaticCastSharedPtr<FOSCBundlePacket>(Bundle.GetPacket());
@@ -201,6 +221,8 @@ void UOSCManager::GetMessagesFromBundle(const FOSCBundle& Bundle, TArray<FOSCMes
 			}
 		}
 	}
+	
+	return Messages;
 }
 
 void UOSCManager::GetFloat(const FOSCMessage& Message, const int32 Index, float& Value)
@@ -375,11 +397,27 @@ FOSCAddress& UOSCManager::OSCAddressPushContainer(FOSCAddress& Address, const FS
 	return Address;
 }
 
+FOSCAddress& UOSCManager::OSCAddressPushContainers(FOSCAddress& Address, const TArray<FString>& ToAppend)
+{
+	Address.PushContainers(ToAppend);
+	return Address;
+}
+
 FString UOSCManager::OSCAddressPopContainer(FOSCAddress& Address)
 {
 	return Address.PopContainer();
 }
 
+TArray<FString> UOSCManager::OSCAddressPopContainers(FOSCAddress& Address, int32 InNumContainers)
+{
+	return Address.PopContainers(InNumContainers);
+}
+
+bool UOSCManager::OSCAddressPathMatchesPattern(const FOSCAddress& InPattern, const FOSCAddress& InPath)
+{
+	return InPattern.Matches(InPath);
+}
+
 FOSCAddress UOSCManager::GetOSCMessageAddress(const FOSCMessage& Message)
 {
 	return Message.GetAddress();
@@ -400,7 +438,7 @@ TArray<FString> UOSCManager::GetOSCAddressContainers(const FOSCAddress& Address)
 {
 	TArray<FString> Containers;
 	Address.GetContainers(Containers);
-	return MoveTemp(Containers);
+	return Containers;
 }
 
 FString UOSCManager::GetOSCAddressContainerPath(const FOSCAddress& Address)
@@ -418,8 +456,14 @@ FString UOSCManager::GetOSCAddressMethod(const FOSCAddress& Address)
 	return Address.GetMethod();
 }
 
+FOSCAddress& UOSCManager::ClearOSCAddressContainers(FOSCAddress& Address)
+{
+	Address.ClearContainers();
+	return Address;
+}
+
 FOSCAddress& UOSCManager::SetOSCAddressMethod(FOSCAddress& Address, const FString& Method)
 {
 	Address.SetMethod(Method);
 	return Address;
-}
+}
\ No newline at end of file
diff --git a/Engine/Plugins/Runtime/OSC/Source/OSC/Public/OSCAddress.h b/Engine/Plugins/Runtime/OSC/Source/OSC/Public/OSCAddress.h
--- a/Engine/Plugins/Runtime/OSC/Source/OSC/Public/OSCAddress.h
+++ b/Engine/Plugins/Runtime/OSC/Source/OSC/Public/OSCAddress.h
@@ -52,14 +52,20 @@ public:
 	  */
 	bool Matches(const FOSCAddress& InAddress) const;
 
-	/** Pushes container onto address' ordered array of containers */
+	/** Pushes container onto end of address' ordered array of containers */
 	void PushContainer(const FString& InContainer);
 
+	/** Pushes containers onto end of address' ordered array of containers */
+	void PushContainers(const TArray<FString>& InContainers);
+
 	/** Pops container from ordered array of containers */
 	FString PopContainer();
 
+	/** Pops containers off end of address' ordered array of containers */
+	TArray<FString> PopContainers(int32 InNumContainers);
+
 	/** Clears ordered array of containers */
-	void ClearContainers(const FString& InContainer);
+	void ClearContainers();
 
 	/** Get method name of address */
 	const FString& GetMethod() const;
diff --git a/Engine/Plugins/Runtime/OSC/Source/OSC/Public/OSCManager.h b/Engine/Plugins/Runtime/OSC/Source/OSC/Public/OSCManager.h
--- a/Engine/Plugins/Runtime/OSC/Source/OSC/Public/OSCManager.h
+++ b/Engine/Plugins/Runtime/OSC/Source/OSC/Public/OSCManager.h
@@ -26,87 +26,91 @@ public:
 	static class UOSCClient* CreateOSCClient(FString SendIPAddress, int32 Port);
 
 	/** Adds provided message packet to bundle. */
-	UFUNCTION(BlueprintCallable, Category = "Audio|OSC")
+	UFUNCTION(BlueprintCallable, Category = "Audio|OSC", meta = (DisplayName = "Add OSC MEssage to Bundle", Keywords = "osc message bundle"))
 	static UPARAM(DisplayName = "Bundle") FOSCBundle& AddMessageToBundle(const FOSCMessage& Message, UPARAM(ref) FOSCBundle& Bundle);
 
 	/** Adds bundle packet to bundle. */
-	UFUNCTION(BlueprintCallable, Category = "Audio|OSC")
+	UFUNCTION(BlueprintCallable, Category = "Audio|OSC", meta = (DisplayName = "Add OSC Bundle to Bundle", Keywords = "osc message"))
 	static UPARAM(DisplayName = "OutBundle") FOSCBundle& AddBundleToBundle(const FOSCBundle& InBundle, UPARAM(ref) FOSCBundle& OutBundle);
 
-	/** Fills array of messages found in bundle. */
-	UFUNCTION(BlueprintCallable, Category = "Audio|OSC")
-	static void GetMessagesFromBundle(const FOSCBundle& Bundle, TArray<FOSCMessage>& Messages);
+	/** Fills array with child bundles found in bundle. */
+	UFUNCTION(BlueprintCallable, Category = "Audio|OSC", meta = (DisplayName = "Get OSC Bundles From Bundle", Keywords = "osc bundle"))
+	static UPARAM(DisplayName = "Bundles") TArray<FOSCBundle>& GetBundlesFromBundle(const FOSCBundle& Bundle, TArray<FOSCBundle>& Bundles);
+
+	/** Fills array with messages found in bundle. */
+	UFUNCTION(BlueprintCallable, Category = "Audio|OSC", meta = (DisplayName = "Get OSC Messages From Bundle", Keywords = "osc bundle message"))
+	static UPARAM(DisplayName = "Messages") TArray<FOSCMessage>& GetMessagesFromBundle(const FOSCBundle& Bundle, TArray<FOSCMessage>& Messages);
 
 	/** Clears provided message of all arguments. */
-	UFUNCTION(BlueprintCallable, Category = "Audio|OSC")
+	UFUNCTION(BlueprintCallable, Category = "Audio|OSC", meta = (DisplayName = "Clear OSC Message", Keywords = "osc message"))
 	static UPARAM(DisplayName = "Message") FOSCMessage& ClearMessage(UPARAM(ref) FOSCMessage& Message);
 
 	/** Clears provided bundle of all internal messages/bundle packets. */
-	UFUNCTION(BlueprintCallable, Category = "Audio|OSC")
+	UFUNCTION(BlueprintCallable, Category = "Audio|OSC", meta = (DisplayName = "Clear OSC Bundle", Keywords = "osc message"))
 	static UPARAM(DisplayName = "Bundle") FOSCBundle& ClearBundle(UPARAM(ref) FOSCBundle& Bundle);
 
 	/** Adds float value to end of OSCMessage */
-	UFUNCTION(BlueprintCallable, Category = "Audio|OSC")
+	UFUNCTION(BlueprintCallable, Category = "Audio|OSC", meta = (DisplayName = "Add Float to OSC Message", Keywords = "osc message"))
 	static UPARAM(DisplayName = "Message") FOSCMessage& AddFloat(UPARAM(ref) FOSCMessage& Message, float Value);
 
 	/** Adds Int32 value to end of OSCMessage */
-	UFUNCTION(BlueprintCallable, Category = "Audio|OSC")
+	UFUNCTION(BlueprintCallable, Category = "Audio|OSC", meta = (DisplayName = "Add Integer to OSC Message", Keywords = "osc message"))
 	static UPARAM(DisplayName = "Message") FOSCMessage& AddInt32(UPARAM(ref) FOSCMessage& Message, int32 Value);
 
 	/** Adds Int64 value to end of OSCMessage */
-	UFUNCTION(BlueprintCallable, Category = "Audio|OSC")
+	UFUNCTION(BlueprintCallable, Category = "Audio|OSC", meta = (DisplayName = "Add Integer (64-bit) to OSC Message", Keywords = "osc message"))
 	static UPARAM(DisplayName = "Message") FOSCMessage& AddInt64(UPARAM(ref) FOSCMessage& Message, int64 Value);
 
 	/** Adds string value to end of OSCMessage */
-	UFUNCTION(BlueprintCallable, Category = "Audio|OSC")
+	UFUNCTION(BlueprintCallable, Category = "Audio|OSC", meta = (DisplayName = "Add String to OSC Message", Keywords = "osc message"))
 	static UPARAM(DisplayName = "Message") FOSCMessage& AddString(UPARAM(ref) FOSCMessage& Message, FString Value);
 
 	/** Adds blob value to end of OSCMessage */
-	UFUNCTION(BlueprintCallable, Category = "Audio|OSC")
+	UFUNCTION(BlueprintCallable, Category = "Audio|OSC", meta = (DisplayName = "Add Blob to OSC Message", Keywords = "osc message"))
 	static UPARAM(DisplayName = "Message") FOSCMessage& AddBlob(UPARAM(ref) FOSCMessage& Message, TArray<uint8>& Value);
 
 	/** Adds boolean value to end of OSCMessage */
-	UFUNCTION(BlueprintCallable, Category = "Audio|OSC")
+	UFUNCTION(BlueprintCallable, Category = "Audio|OSC", meta = (DisplayName = "Add Bool to OSC Message", Keywords = "osc message"))
 	static UPARAM(DisplayName = "Message") FOSCMessage& AddBool(UPARAM(ref) FOSCMessage& Message, bool Value);
 
 	/** Set Value to float at provided Index in OSCMessage if in bounds and type matches */
-	UFUNCTION(BlueprintCallable, Category = "Audio|OSC")
+	UFUNCTION(BlueprintCallable, Category = "Audio|OSC", meta = (DisplayName = "Get OSC Message Float At Index", Keywords = "osc message"))
 	static void GetFloat(const FOSCMessage& Message, const int32 Index, float& Value);
 
 	/** Returns all float values in order of received from OSCMessage */
-	UFUNCTION(BlueprintCallable, Category = "Audio|OSC")
+	UFUNCTION(BlueprintCallable, Category = "Audio|OSC", meta = (DisplayName = "Get OSC Message Floats", Keywords = "osc message"))
 	static void GetAllFloats(const FOSCMessage& Message, TArray<float>& Values);
 
 	/** Set Value to integer at provided Index in OSCMessage if in bounds and type matches */
-	UFUNCTION(BlueprintCallable, Category = "Audio|OSC")
+	UFUNCTION(BlueprintCallable, Category = "Audio|OSC", meta = (DisplayName = "Get OSC Message Integer at Index", Keywords = "osc message"))
 	static void GetInt32(const FOSCMessage& Message, const int32 Index, int32& Value);
 
 	/** Returns all integer values in order of received from OSCMessage */
-	UFUNCTION(BlueprintCallable, Category = "Audio|OSC")
+	UFUNCTION(BlueprintCallable, Category = "Audio|OSC", meta = (DisplayName = "Get OSC Message Integers", Keywords = "osc message"))
 	static void GetAllInt32s(const FOSCMessage& Message, TArray<int32>& Values);
 
 	/** Set Value to Int64 at provided Index in OSCMessage if in bounds and type matches */
-	UFUNCTION(BlueprintCallable, Category = "Audio|OSC")
+	UFUNCTION(BlueprintCallable, Category = "Audio|OSC", meta = (DisplayName = "Get OSC Message Integer (64-bit) at Index", Keywords = "osc message"))
 	static void GetInt64(const FOSCMessage& Message, const int32 Index, int64& Value);
 
 	/** Returns all Int64 values in order of received from OSCMessage */
-	UFUNCTION(BlueprintCallable, Category = "Audio|OSC")
+	UFUNCTION(BlueprintCallable, Category = "Audio|OSC", meta = (DisplayName = "Get OSC Message Integers (64-bit)", Keywords = "osc message"))
 	static void GetAllInt64s(const FOSCMessage& Message, TArray<int64>& Values);
 
 	/** Set Value to string at provided Index in OSCMessage if in bounds and type matches */
-	UFUNCTION(BlueprintCallable, Category = "Audio|OSC")
+	UFUNCTION(BlueprintCallable, Category = "Audio|OSC", meta = (DisplayName = "Get OSC Message String at Index", Keywords = "osc message"))
 	static void GetString(const FOSCMessage& Message, const int32 Index, FString& Value);
 
 	/** Returns all string values in order of received from OSCMessage */
-	UFUNCTION(BlueprintCallable, Category = "Audio|OSC")
+	UFUNCTION(BlueprintCallable, Category = "Audio|OSC", meta = (DisplayName = "Get OSC Message Strings", Keywords = "osc message"))
 	static void GetAllStrings(const FOSCMessage& Message, TArray<FString>& Values);
 
 	/** Sets Value to boolean at provided Index from OSCMessage if in bounds and type matches */
-	UFUNCTION(BlueprintCallable, Category = "Audio|OSC")
+	UFUNCTION(BlueprintCallable, Category = "Audio|OSC", meta = (DisplayName = "Get OSC Message Bool At Index", Keywords = "osc message"))
 	static void GetBool(const FOSCMessage& Message, const int32 Index, bool& Value);
 
 	/** Returns all boolean values in order of received from OSCMessage */
-	UFUNCTION(BlueprintCallable, Category = "Audio|OSC")
+	UFUNCTION(BlueprintCallable, Category = "Audio|OSC", meta = (DisplayName = "Get OSC Message Bools", Keywords = "osc message"))
 	static void GetAllBools(const FOSCMessage& Message, TArray<bool>& Values);
 
 	/** Sets Value to blob at provided Index from OSCMessage if in bounds and type matches */
@@ -114,55 +118,73 @@ public:
 	static void GetBlob(const FOSCMessage& Message, const int32 Index, TArray<uint8>& Value);
 
 	/** Returns whether OSCAddress is valid path */
-	UFUNCTION(BlueprintCallable, Category = "Audio|OSC")
+	UFUNCTION(BlueprintCallable, Category = "Audio|OSC", meta = (DisplayName = "Is OSC Address Valid Path", Keywords = "valid osc address path address"))
 	static bool OSCAddressIsValidPath(const FOSCAddress& Address);
 
 	/** Returns whether OSCAddress is valid pattern to match against */
-	UFUNCTION(BlueprintCallable, Category = "Audio|OSC")
+	UFUNCTION(BlueprintCallable, Category = "Audio|OSC", meta = (DisplayName = "Is OSC Address Valid Pattern", Keywords = "valid osc address pattern address"))
 	static bool OSCAddressIsValidPattern(const FOSCAddress& Address);
 
 	/* Converts string to OSCAddress */
 	UFUNCTION(BlueprintCallable, Category = "Audio|OSC")
 	static UPARAM(DisplayName = "Address") FOSCAddress ConvertStringToOSCAddress(const FString& String);
 
-	/* Pushes container on end of OSCAddress' ordered container array */
-	UFUNCTION(BlueprintCallable, Category = "Audio|OSC")
+	/** Returns if address pattern matches the provided address path.
+	  * If passed address is not a valid path, returns false.
+	  */
+	UFUNCTION(BlueprintCallable, Category = "Audio|OSC", meta = (DisplayName = "OSC Address Path Matches Pattern", Keywords = "matches osc address path address"))
+	static bool UPARAM(DisplayName = "Is Match") OSCAddressPathMatchesPattern(const FOSCAddress& Pattern, const FOSCAddress& Path);
+
+	/** Pushes container onto address' ordered array of containers */
+	UFUNCTION(BlueprintCallable, Category = "Audio|OSC", meta = (DisplayName = "Push Container to OSC Address", Keywords = "push osc address container"))
 	static UPARAM(DisplayName = "Address") FOSCAddress& OSCAddressPushContainer(UPARAM(ref) FOSCAddress& Address, const FString& Container);
 
-	/* Pops container off end of OSCAddress' ordered container array */
-	UFUNCTION(BlueprintCallable, Category = "Audio|OSC")
-	static UPARAM(DisplayName = "Address") FString OSCAddressPopContainer(UPARAM(ref) FOSCAddress& Address);
+	/** Pushes container onto address' ordered array of containers */
+	UFUNCTION(BlueprintCallable, Category = "Audio|OSC", meta = (DisplayName = "Push Container Array to OSC Address", Keywords = "push osc address container"))
+	static UPARAM(DisplayName = "Address") FOSCAddress& OSCAddressPushContainers(UPARAM(ref) FOSCAddress& Address, const TArray<FString>& Containers);
+
+	/** Pops container from ordered array of containers. If no containers, returns empty string */
+	UFUNCTION(BlueprintCallable, Category = "Audio|OSC", meta = (DisplayName = "Pop Container from OSC Address", Keywords = "pop osc address container"))
+	static UPARAM(DisplayName = "Container") FString OSCAddressPopContainer(UPARAM(ref) FOSCAddress& Address);
+
+	/** Pops container from ordered array of containers. If NumContainers is greater than or equal to the number of containers in address, returns all containers. */
+	UFUNCTION(BlueprintCallable, Category = "Audio|OSC", meta = (DisplayName = "Pop Containers from OSC Address", Keywords = "pop osc address container"))
+	static UPARAM(DisplayName = "Containers") TArray<FString> OSCAddressPopContainers(UPARAM(ref) FOSCAddress& Address, int32 NumContainers);
 
 	/* Returns copy of message's OSCAddress */
 	UFUNCTION(BlueprintCallable, Category = "Audio|OSC")
 	static UPARAM(DisplayName = "Address") FOSCAddress GetOSCMessageAddress(const FOSCMessage& Message);
 
-	/** Sets the OSCAdress Address of the provided message */
+	/** Sets the OSCAddress of the provided message */
 	UFUNCTION(BlueprintCallable, Category = "Audio|OSC")
 	static UPARAM(DisplayName = "Message") FOSCMessage& SetOSCMessageAddress(UPARAM(ref) FOSCMessage& Message, const FOSCAddress& Address);
 
 	/** Returns the OSCAddress container at the provided 'Index.' Returns empty string if index is out-of-bounds. */
-	UFUNCTION(BlueprintCallable, Category = "Audio|OSC")
+	UFUNCTION(BlueprintCallable, Category = "Audio|OSC", meta = (DisplayName = "Get OSC Address Container At Index", Keywords = "osc address container path"))
 	static UPARAM(DisplayName = "Container") FString GetOSCAddressContainer(const FOSCAddress& Address, const int32 Index);
 
-	/** Returns ordered array of OSCAddress containers within OSCAddress' path. */
-	UFUNCTION(BlueprintCallable, Category = "Audio|OSC")
+	/** Builds referenced array of address of containers in order */
+	UFUNCTION(BlueprintCallable, Category = "Audio|OSC", meta = (DisplayName = "Get OSC Address Containers", Keywords = "osc address container path"))
 	static UPARAM(DisplayName = "Containers") TArray<FString> GetOSCAddressContainers(const FOSCAddress& Address);
 
-	/** Returns OSCAddress' container path. */
+	/** Returns full path of OSC address in the form '/Container1/Container2/Method' */
 	UFUNCTION(BlueprintCallable, Category = "Audio|OSC")
 	static UPARAM(DisplayName = "Path") FString GetOSCAddressContainerPath(const FOSCAddress& Address);
 
-	/** Returns OSCAddress' full path. (Container path plus method name) */
-	UFUNCTION(BlueprintCallable, Category = "Audio|OSC")
+	/** Returns full path of OSC address in the form '/Container1/Container2' */
+	UFUNCTION(BlueprintCallable, Category = "Audio|OSC", meta = (DisplayName = "Convert OSC Address To String", Keywords = "osc address path"))
 	static UPARAM(DisplayName = "Full Path") FString GetOSCAddressFullPath(const FOSCAddress& Address);
 
 	/** Returns method name of OSCAddress provided */
 	UFUNCTION(BlueprintCallable, Category = "Audio|OSC")
 	static UPARAM(DisplayName = "Method") FString GetOSCAddressMethod(const FOSCAddress& Address);
 
+	/** Clears containers of OSCAddress provided */
+	UFUNCTION(BlueprintCallable, Category = "Audio|OSC", meta = (DisplayName = "Clear OSC Address Containers", Keywords = "osc address clear"))
+	static UPARAM(DisplayName = "Address") FOSCAddress& ClearOSCAddressContainers(FOSCAddress& Address);
+
 	/** Sets the method name of the OSCAddress provided */
-	UFUNCTION(BlueprintCallable, Category = "Audio|OSC")
+	UFUNCTION(BlueprintCallable, Category = "Audio|OSC", meta = (DisplayName = "Set OSC Address Method", Keywords = "osc method"))
 	static UPARAM(DisplayName = "Address") FOSCAddress& SetOSCAddressMethod(FOSCAddress& Address, const FString& Method);
 };
 
