Fix compiler error
#rnx

diff --git a/Engine/Plugins/Runtime/OSC/Source/OSC/Public/OSCManager.h b/Engine/Plugins/Runtime/OSC/Source/OSC/Public/OSCManager.h
--- a/Engine/Plugins/Runtime/OSC/Source/OSC/Public/OSCManager.h
+++ b/Engine/Plugins/Runtime/OSC/Source/OSC/Public/OSCManager.h
@@ -133,7 +133,7 @@ public:
 	  * If passed address is not a valid path, returns false.
 	  */
 	UFUNCTION(BlueprintCallable, Category = "Audio|OSC", meta = (DisplayName = "OSC Address Path Matches Pattern", Keywords = "matches osc address path address"))
-	static bool UPARAM(DisplayName = "Is Match") OSCAddressPathMatchesPattern(const FOSCAddress& Pattern, const FOSCAddress& Path);
+	static UPARAM(DisplayName = "Is Match") bool OSCAddressPathMatchesPattern(const FOSCAddress& Pattern, const FOSCAddress& Path);
 
 	/** Pushes container onto address' ordered array of containers */
 	UFUNCTION(BlueprintCallable, Category = "Audio|OSC", meta = (DisplayName = "Push Container to OSC Address", Keywords = "push osc address container"))
