- Enable OSC for all platforms
- Fix null termination memory scribble/crash when passing empty strings as payload in message members
- Fix memory scribble when sending messages larger than 1024 bytes

Minor features:
- Clean-up some OSC BP interface issues (missing UPARAM(ref) markup, comments, etc.)
- Add bool return to message payload array requests to determine if failed or if value is "fresh"
- Add BP functions to convert Object path to and from OSCAddress
- Add ability to remove OSC containers at index
- Fix pattern matching for wildcarding over span of containers/method

diff --git a/Engine/Plugins/Runtime/OSC/OSC.uplugin b/Engine/Plugins/Runtime/OSC/OSC.uplugin
--- a/Engine/Plugins/Runtime/OSC/OSC.uplugin
+++ b/Engine/Plugins/Runtime/OSC/OSC.uplugin
@@ -1,9 +1,9 @@
 {
 	"FileVersion" : 3,
 	"Version" : 1,
-	"VersionName" : "1.0",
-	"FriendlyName" : "OSC Support",
-	"Description" : "Allows you to send and receive OSC events through a simple API in either C++ or Blueprints",
+	"VersionName" : "1.1",
+	"FriendlyName" : "OSC (Open Sound Control)",
+	"Description" : "Implements the OSC 1.0 specification, allowing users to send and receive OSC messages and bundles between remote clients or applications.",
 	"Category" : "Input Devices",
 	"CreatedBy" : "Epic Games, Inc.",
 	"CreatedByURL" : "http://epicgames.com",
@@ -12,18 +12,14 @@
 	"SupportURL" : "",
 	"EnabledByDefault" : false,
 	"CanContainContent" : false,
-	"IsBetaVersion" : true,
+	"IsBetaVersion" : false,
 	"Installed" : false,
 	"Modules" :
 	[
 		{
 			"Name": "OSC",
 			"Type": "Runtime",
-			"LoadingPhase": "Default",
-			"WhitelistPlatforms": [
-				"Win64",
-				"Win32"
-			]
+			"LoadingPhase": "Default"
 		}
 	]
 }
\ No newline at end of file
diff --git a/Engine/Plugins/Runtime/OSC/Source/OSC/Private/OSCAddress.cpp b/Engine/Plugins/Runtime/OSC/Source/OSC/Private/OSCAddress.cpp
--- a/Engine/Plugins/Runtime/OSC/Source/OSC/Private/OSCAddress.cpp
+++ b/Engine/Plugins/Runtime/OSC/Source/OSC/Private/OSCAddress.cpp
@@ -45,25 +45,7 @@ bool FOSCAddress::Matches(const FOSCAddress& InAddress) const
 {
 	if (IsValidPattern() && InAddress.IsValidPath())
 	{
-		if (Containers.Num() != InAddress.Containers.Num())
-		{
-			return false;
-		}
-
-		if (!FAudioAddressPattern::PartsMatch(Method, InAddress.Method))
-		{
-			return false;
-		}
-
-		for (int32 i = 0; i < Containers.Num(); ++i)
-		{
-			if (!FAudioAddressPattern::PartsMatch(Containers[i], InAddress.Containers[i]))
-			{
-				return false;
-			}
-		}
-
-		return true; 
+		return FAudioAddressPattern::PartsMatch(GetFullPath(), InAddress.GetFullPath());
 	}
 
 	return false;
@@ -147,6 +129,17 @@ TArray<FString> FOSCAddress::PopContainers(int32 InNumContainers)
 	return Popped;
 }
 
+void FOSCAddress::RemoveContainers(int32 InIndex, int32 InCount)
+{
+	if (InIndex >= 0 && InCount > 0)
+	{
+		if (InIndex + InCount < Containers.Num())
+		{
+			Containers.RemoveAt(InIndex, InCount);
+		}
+	}
+}
+
 void FOSCAddress::ClearContainers()
 {
 	Containers.Reset();
diff --git a/Engine/Plugins/Runtime/OSC/Source/OSC/Private/OSCBundlePacket.cpp b/Engine/Plugins/Runtime/OSC/Source/OSC/Private/OSCBundlePacket.cpp
--- a/Engine/Plugins/Runtime/OSC/Source/OSC/Private/OSCBundlePacket.cpp
+++ b/Engine/Plugins/Runtime/OSC/Source/OSC/Private/OSCBundlePacket.cpp
@@ -50,6 +50,8 @@ void FOSCBundlePacket::WriteData(FOSCStream& Stream)
 
 void FOSCBundlePacket::ReadData(FOSCStream& Stream)
 {
+	Packets.Reset();
+
 	FString BundleTag = Stream.ReadString();
 	if (BundleTag != OSC::BundleTag)
 	{
diff --git a/Engine/Plugins/Runtime/OSC/Source/OSC/Private/OSCBundlePacket.h b/Engine/Plugins/Runtime/OSC/Source/OSC/Private/OSCBundlePacket.h
--- a/Engine/Plugins/Runtime/OSC/Source/OSC/Private/OSCBundlePacket.h
+++ b/Engine/Plugins/Runtime/OSC/Source/OSC/Private/OSCBundlePacket.h
@@ -7,7 +7,7 @@
 #include "OSCTypes.h"
 
 
-class FOSCBundlePacket : public IOSCPacket
+class OSC_API FOSCBundlePacket : public IOSCPacket
 {
 public:
 	FOSCBundlePacket();
diff --git a/Engine/Plugins/Runtime/OSC/Source/OSC/Private/OSCClient.cpp b/Engine/Plugins/Runtime/OSC/Source/OSC/Private/OSCClient.cpp
--- a/Engine/Plugins/Runtime/OSC/Source/OSC/Private/OSCClient.cpp
+++ b/Engine/Plugins/Runtime/OSC/Source/OSC/Private/OSCClient.cpp
@@ -24,6 +24,11 @@ void UOSCClient::Connect()
 	ClientProxy.Reset(new FOSCClientProxy(GetName()));
 }
 
+bool UOSCClient::IsActive() const
+{
+	return ClientProxy.IsValid() && ClientProxy->IsActive();
+}
+
 void UOSCClient::GetSendIPAddress(FString& InIPAddress, int32& Port)
 {
 	check(ClientProxy.IsValid());
diff --git a/Engine/Plugins/Runtime/OSC/Source/OSC/Private/OSCClientProxy.cpp b/Engine/Plugins/Runtime/OSC/Source/OSC/Private/OSCClientProxy.cpp
--- a/Engine/Plugins/Runtime/OSC/Source/OSC/Private/OSCClientProxy.cpp
+++ b/Engine/Plugins/Runtime/OSC/Source/OSC/Private/OSCClientProxy.cpp
@@ -3,6 +3,7 @@
 
 #include "Common/UdpSocketBuilder.h"
 #include "Sockets.h"
+#include "SocketTypes.h"
 
 #include "OSCBundlePacket.h"
 #include "OSCLog.h"
@@ -48,6 +49,11 @@ bool FOSCClientProxy::SetSendIPAddress(const FString& InIPAddress, const int32 P
 	return bIsValidAddress;
 }
 
+bool FOSCClientProxy::IsActive() const
+{
+	return Socket && Socket->GetConnectionState() == ESocketConnectionState::SCS_Connected;
+}
+
 void FOSCClientProxy::SendPacket(IOSCPacket& Packet)
 {
 	if (!Socket)
@@ -70,9 +76,9 @@ void FOSCClientProxy::SendPacket(IOSCPacket& Packet)
 
 	if (IPAddress)
 	{
-		FOSCStream Stream = FOSCStream(OUTPUT_BUFFER_SIZE);
-		const uint8* DataPtr = Stream.GetData();
+		FOSCStream Stream = FOSCStream();
 		Packet.WriteData(Stream);
+		const uint8* DataPtr = Stream.GetData();
 
 		int32 BytesSent = 0;
 
diff --git a/Engine/Plugins/Runtime/OSC/Source/OSC/Private/OSCClientProxy.h b/Engine/Plugins/Runtime/OSC/Source/OSC/Private/OSCClientProxy.h
--- a/Engine/Plugins/Runtime/OSC/Source/OSC/Private/OSCClientProxy.h
+++ b/Engine/Plugins/Runtime/OSC/Source/OSC/Private/OSCClientProxy.h
@@ -18,6 +18,8 @@ public:
 	void GetSendIPAddress(FString& InIPAddress, int32& Port) const override;
 	bool SetSendIPAddress(const FString& InIPAddress, const int32 Port) override;
 
+	bool IsActive() const override;
+
 	void SendMessage(FOSCMessage& Message) override;
 	void SendBundle(FOSCBundle& Bundle) override;
 
diff --git a/Engine/Plugins/Runtime/OSC/Source/OSC/Private/OSCManager.cpp b/Engine/Plugins/Runtime/OSC/Source/OSC/Private/OSCManager.cpp
--- a/Engine/Plugins/Runtime/OSC/Source/OSC/Private/OSCManager.cpp
+++ b/Engine/Plugins/Runtime/OSC/Source/OSC/Private/OSCManager.cpp
@@ -7,16 +7,25 @@
 #include "OSCBundle.h"
 #include "OSCBundlePacket.h"
 #include "OSCLog.h"
+#include "OSCServer.h"
+#include "OSCClient.h"
+
+#include "Engine/World.h"
+#include "Logging/LogMacros.h"
+#include "Misc/Paths.h"
 #include "SocketSubsystem.h"
+#include "UObject/UObjectIterator.h"
+
 
+#define OSC_LOG_INVALID_TYPE_AT_INDEX(Type, Index, Msg) UE_LOG(LogOSC, Warning, TEXT("OSC Message Parse Failed: OSCType not %s: index '%i', OSCAddress '%s'"), TEXT(##Type), Index, *Msg.GetAddress().GetFullPath())
 
-namespace
+namespace OSC
 {
 	// Returns true if provided address was null and was able to
 	// override with local host address, false if not.
-	bool GetLocalHostAddress(FString& Address)
+	bool GetLocalHostAddress(FString& InAddress)
 	{
-		if (!Address.IsEmpty() && Address != TEXT("0"))
+		if (!InAddress.IsEmpty() && InAddress != TEXT("0"))
 		{
 			return false;
 		}
@@ -28,7 +37,7 @@ namespace
 			const TSharedPtr<FInternetAddr> Addr = SocketSys->GetLocalHostAddr(*GLog, bCanBind);
 			if (Addr.IsValid())
 			{
-				Address = Addr->ToString(bAppendPort);
+				InAddress = Addr->ToString(bAppendPort);
 				return true;
 			}
 		}
@@ -36,93 +45,174 @@ namespace
 		return false;
 	}
 
-	const FOSCType* GetOSCTypeAtIndex(const FOSCMessage& Message, const int32 Index)
+	const FOSCType* GetOSCTypeAtIndex(const FOSCMessage& InMessage, const int32 InIndex)
 	{
-		const TSharedPtr<FOSCMessagePacket>& Packet = StaticCastSharedPtr<FOSCMessagePacket>(Message.GetPacket());
+		const TSharedPtr<FOSCMessagePacket>& Packet = StaticCastSharedPtr<FOSCMessagePacket>(InMessage.GetPacket());
 		if (Packet.IsValid())
 		{
 			TArray<FOSCType>& Args = Packet->GetArguments();
-			if (Index >= Args.Num())
+			if (InIndex >= Args.Num())
 			{
-				UE_LOG(LogOSC, Error, TEXT("Index '%d' out-of-bounds.  Message argument size = '%d'"), Index, Args.Num());
+				UE_LOG(LogOSC, Warning, TEXT("Index '%d' out-of-bounds.  Message argument size = '%d'"), InIndex, Args.Num());
 				return nullptr;
 			}
 
-			return &Args[Index];
+			return &Args[InIndex];
 		}
 
 		return nullptr;
 	}
-} // namespace <>
+} // namespace OSC
+
 
-UOSCServer* UOSCManager::CreateOSCServer(FString ReceiveIPAddress, int32 Port, bool bMulticastLoopback, bool bStartListening)
+static FAutoConsoleCommand GOSCGetServerDiag(
+	TEXT("osc.servers"),
+	TEXT("Prints diagnostic information pertaining to the current OSC servers to the output log."),
+	FConsoleCommandDelegate::CreateStatic(
+		[]()
+		{
+			FString LocalAddr;
+			OSC::GetLocalHostAddress(LocalAddr);
+			UE_LOG(LogOSC, Display, TEXT("Local IP: %s"), *LocalAddr);
+
+			UE_LOG(LogOSC, Display, TEXT("OSC Servers:"));
+			for (TObjectIterator<UOSCServer> Itr; Itr; ++Itr)
+			{
+				if (UOSCServer* Server = *Itr)
+				{
+					FString ToPrint = TEXT("    ") + Server->GetName();
+					if (UWorld* World = Server->GetWorld())
+					{
+						ToPrint += TEXT("(") + World->GetName() + TEXT(")");
+					}
+
+					ToPrint += TEXT(", ");
+					ToPrint += Server->GetIpAddress(true /* bIncludePort */);
+					ToPrint += Server->IsActive() ? TEXT(" [Active]") : TEXT(" [Inactive]");
+
+					UE_LOG(LogOSC, Display, TEXT("%s"), *ToPrint);
+
+					const TArray<FOSCAddress> BoundPatterns = Server->GetBoundOSCAddressPatterns();
+					if (BoundPatterns.Num() > 0)
+					{
+						UE_LOG(LogOSC, Display, TEXT("\n    Bound Address Patterns:"));
+						for (const FOSCAddress& Pattern : BoundPatterns)
+						{
+							UE_LOG(LogOSC, Display, TEXT("\n         %s"), *Pattern.GetFullPath());
+						}
+						UE_LOG(LogOSC, Display, TEXT(""));
+					}
+				}
+			}
+		}
+	)
+);
+
+static FAutoConsoleCommand GOSCGetClientDiag(
+	TEXT("osc.clients"),
+	TEXT("Prints diagnostic information pertaining to the current OSC clients to the output log."),
+	FConsoleCommandDelegate::CreateStatic(
+		[]()
+		{
+			FString LocalAddr;
+			OSC::GetLocalHostAddress(LocalAddr);
+			UE_LOG(LogOSC, Display, TEXT("Local IP: %s"), *LocalAddr);
+
+			UE_LOG(LogOSC, Display, TEXT("OSC Clients:"));
+			for (TObjectIterator<UOSCClient> Itr; Itr; ++Itr)
+			{
+				if (UOSCClient* Client = *Itr)
+				{
+					FString ToPrint;
+					ToPrint += TEXT("\n    ") + Client->GetName();
+
+					if (UWorld* World = Client->GetWorld())
+					{
+						ToPrint += TEXT("(") + World->GetName() + TEXT(")");
+					}
+
+					FString IPAddrStr;
+					int32 Port;
+					Client->GetSendIPAddress(IPAddrStr, Port);
+					ToPrint += TEXT(", ") + IPAddrStr + TEXT(":");
+					ToPrint.AppendInt(Port);
+					ToPrint += Client->IsActive() ? TEXT(" [Active]") : TEXT(" [Inactive]");
+					UE_LOG(LogOSC, Display, TEXT("%s"), *ToPrint);
+				}
+			}
+		}
+	)
+);
+
+
+UOSCServer* UOSCManager::CreateOSCServer(FString InReceiveIPAddress, int32 InPort, bool bInMulticastLoopback, bool bInStartListening)
 {
-	if (GetLocalHostAddress(ReceiveIPAddress))
+	if (OSC::GetLocalHostAddress(InReceiveIPAddress))
 	{
-		UE_LOG(LogOSC, Display, TEXT("OSCServer ReceiveAddress not specified. Using LocalHost IP: '%s'"), *ReceiveIPAddress);
+		UE_LOG(LogOSC, Display, TEXT("OSCServer ReceiveAddress not specified. Using LocalHost IP: '%s'"), *InReceiveIPAddress);
 	}
 
 	UOSCServer* NewOSCServer = NewObject<UOSCServer>();
 	NewOSCServer->Connect();
-	NewOSCServer->SetMulticastLoopback(bMulticastLoopback);
-	if (NewOSCServer->SetAddress(ReceiveIPAddress, Port))
+	NewOSCServer->SetMulticastLoopback(bInMulticastLoopback);
+	if (NewOSCServer->SetAddress(InReceiveIPAddress, InPort))
 	{
-		if (bStartListening)
+		if (bInStartListening)
 		{
 			NewOSCServer->Listen();
 		}
 	}
 	else
 	{
-		UE_LOG(LogOSC, Warning, TEXT("Failed to parse ReceiveAddress '%s' for OSCServer."), *ReceiveIPAddress);
+		UE_LOG(LogOSC, Warning, TEXT("Failed to parse ReceiveAddress '%s' for OSCServer."), *InReceiveIPAddress);
 	}
 
 	return NewOSCServer;
 }
 
-UOSCClient* UOSCManager::CreateOSCClient(FString SendIPAddress, int32 Port)
+UOSCClient* UOSCManager::CreateOSCClient(FString InSendIPAddress, int32 InPort)
 {
-	if (GetLocalHostAddress(SendIPAddress))
+	if (OSC::GetLocalHostAddress(InSendIPAddress))
 	{
-		UE_LOG(LogOSC, Display, TEXT("OSCClient SendAddress not specified. Using LocalHost IP: '%s'"), *SendIPAddress);
+		UE_LOG(LogOSC, Display, TEXT("OSCClient SendAddress not specified. Using LocalHost IP: '%s'"), *InSendIPAddress);
 	}
 
 	UOSCClient* NewOSCClient = NewObject<UOSCClient>();
 	NewOSCClient->Connect();
-	if (!NewOSCClient->SetSendIPAddress(SendIPAddress, Port))
+	if (!NewOSCClient->SetSendIPAddress(InSendIPAddress, InPort))
 	{
-		UE_LOG(LogOSC, Warning, TEXT("Failed to parse SendAddress '%s' for OSCClient. Client unable to send new messages."), *SendIPAddress);
+		UE_LOG(LogOSC, Warning, TEXT("Failed to parse SendAddress '%s' for OSCClient. Client unable to send new messages."), *InSendIPAddress);
 	}
 
 	return NewOSCClient;
 }
 
-FOSCMessage& UOSCManager::ClearMessage(FOSCMessage& Message)
+FOSCMessage& UOSCManager::ClearMessage(FOSCMessage& OutMessage)
 {
-	const TSharedPtr<FOSCMessagePacket>& Packet = StaticCastSharedPtr<FOSCMessagePacket>(Message.GetPacket());
+	const TSharedPtr<FOSCMessagePacket>& Packet = StaticCastSharedPtr<FOSCMessagePacket>(OutMessage.GetPacket());
 	if (Packet.IsValid())
 	{
 		Packet->GetArguments().Reset();
 	}
 
-	return Message;
+	return OutMessage;
 }
 
-FOSCBundle& UOSCManager::ClearBundle(FOSCBundle& Bundle)
+FOSCBundle& UOSCManager::ClearBundle(FOSCBundle& OutBundle)
 {
-	const TSharedPtr<FOSCBundlePacket>& Packet = StaticCastSharedPtr<FOSCBundlePacket>(Bundle.GetPacket());
+	const TSharedPtr<FOSCBundlePacket>& Packet = StaticCastSharedPtr<FOSCBundlePacket>(OutBundle.GetPacket());
 	if (Packet.IsValid())
 	{
 		Packet->GetPackets().Reset();
 	}
 
-	return Bundle;
+	return OutBundle;
 }
 
-FOSCBundle& UOSCManager::AddMessageToBundle(const FOSCMessage& Message, FOSCBundle& Bundle)
+FOSCBundle& UOSCManager::AddMessageToBundle(const FOSCMessage& InMessage, FOSCBundle& Bundle)
 {
 	const TSharedPtr<FOSCBundlePacket>& BundlePacket = StaticCastSharedPtr<FOSCBundlePacket>(Bundle.GetPacket());
-	const TSharedPtr<FOSCMessagePacket>& MessagePacket = StaticCastSharedPtr<FOSCMessagePacket>(Message.GetPacket());
+	const TSharedPtr<FOSCMessagePacket>& MessagePacket = StaticCastSharedPtr<FOSCMessagePacket>(InMessage.GetPacket());
 
 	if (MessagePacket.IsValid() && BundlePacket.IsValid())
 	{
@@ -145,54 +235,61 @@ FOSCBundle& UOSCManager::AddBundleToBundle(const FOSCBundle& InBundle, FOSCBundl
 	return OutBundle;
 }
 
-FOSCMessage& UOSCManager::AddFloat(FOSCMessage& Message, float Value)
+FOSCMessage& UOSCManager::AddFloat(FOSCMessage& OutMessage, float InValue)
 {
-	const TSharedPtr<FOSCMessagePacket>& MessagePacket = StaticCastSharedPtr<FOSCMessagePacket>(Message.GetPacket());
-	MessagePacket->GetArguments().Add(FOSCType(Value));
-	return Message;
+	const TSharedPtr<FOSCMessagePacket>& MessagePacket = StaticCastSharedPtr<FOSCMessagePacket>(OutMessage.GetPacket());
+	MessagePacket->GetArguments().Add(FOSCType(InValue));
+	return OutMessage;
 }
 
-FOSCMessage& UOSCManager::AddInt32(FOSCMessage& Message, int32 Value)
+FOSCMessage& UOSCManager::AddInt32(FOSCMessage& OutMessage, int32 InValue)
 {
-	const TSharedPtr<FOSCMessagePacket>& MessagePacket = StaticCastSharedPtr<FOSCMessagePacket>(Message.GetPacket());
-	MessagePacket->GetArguments().Add(FOSCType(Value));
-	return Message;
+	const TSharedPtr<FOSCMessagePacket>& MessagePacket = StaticCastSharedPtr<FOSCMessagePacket>(OutMessage.GetPacket());
+	MessagePacket->GetArguments().Add(FOSCType(InValue));
+	return OutMessage;
 }
 
-FOSCMessage& UOSCManager::AddInt64(FOSCMessage& Message, int64 Value)
+FOSCMessage& UOSCManager::AddInt64(FOSCMessage& OutMessage, int64 InValue)
 {
-	const TSharedPtr<FOSCMessagePacket>& MessagePacket = StaticCastSharedPtr<FOSCMessagePacket>(Message.GetPacket());
-	MessagePacket->GetArguments().Add(FOSCType(Value));
-	return Message;
+	const TSharedPtr<FOSCMessagePacket>& MessagePacket = StaticCastSharedPtr<FOSCMessagePacket>(OutMessage.GetPacket());
+	MessagePacket->GetArguments().Add(FOSCType(InValue));
+	return OutMessage;
 }
 
-FOSCMessage& UOSCManager::AddString(FOSCMessage& Message, FString Value)
+FOSCMessage& UOSCManager::AddAddress(FOSCMessage& OutMessage, const FOSCAddress& InValue)
 {
-	const TSharedPtr<FOSCMessagePacket>& MessagePacket = StaticCastSharedPtr<FOSCMessagePacket>(Message.GetPacket());
-	MessagePacket->GetArguments().Add(FOSCType(Value));
-	return Message;
+	const TSharedPtr<FOSCMessagePacket>& MessagePacket = StaticCastSharedPtr<FOSCMessagePacket>(OutMessage.GetPacket());
+	MessagePacket->GetArguments().Add(FOSCType(InValue.GetFullPath()));
+	return OutMessage;
 }
 
-FOSCMessage& UOSCManager::AddBlob(FOSCMessage& Message, TArray<uint8>& Value)
+FOSCMessage& UOSCManager::AddString(FOSCMessage& OutMessage, FString InValue)
 {
-	const TSharedPtr<FOSCMessagePacket>& MessagePacket = StaticCastSharedPtr<FOSCMessagePacket>(Message.GetPacket());
-	MessagePacket->GetArguments().Add(FOSCType(Value));
-	return Message;
+	const TSharedPtr<FOSCMessagePacket>& MessagePacket = StaticCastSharedPtr<FOSCMessagePacket>(OutMessage.GetPacket());
+	MessagePacket->GetArguments().Add(FOSCType(InValue));
+	return OutMessage;
 }
 
-FOSCMessage& UOSCManager::AddBool(FOSCMessage& Message, bool Value)
+FOSCMessage& UOSCManager::AddBlob(FOSCMessage& OutMessage, TArray<uint8>& OutValue)
 {
-	const TSharedPtr<FOSCMessagePacket>& MessagePacket = StaticCastSharedPtr<FOSCMessagePacket>(Message.GetPacket());
-	MessagePacket->GetArguments().Add(FOSCType(Value));
-	return Message;
+	const TSharedPtr<FOSCMessagePacket>& MessagePacket = StaticCastSharedPtr<FOSCMessagePacket>(OutMessage.GetPacket());
+	MessagePacket->GetArguments().Add(FOSCType(OutValue));
+	return OutMessage;
 }
 
-TArray<FOSCBundle>& UOSCManager::GetBundlesFromBundle(const FOSCBundle& Bundle, TArray<FOSCBundle>& Bundles)
+FOSCMessage& UOSCManager::AddBool(FOSCMessage& OutMessage, bool InValue)
 {
-	Bundles.Reset();
-	if (Bundle.GetPacket().IsValid())
+	const TSharedPtr<FOSCMessagePacket>& MessagePacket = StaticCastSharedPtr<FOSCMessagePacket>(OutMessage.GetPacket());
+	MessagePacket->GetArguments().Add(FOSCType(InValue));
+	return OutMessage;
+}
+
+TArray<FOSCBundle> UOSCManager::GetBundlesFromBundle(const FOSCBundle& InBundle)
+{
+	TArray<FOSCBundle> Bundles;
+	if (InBundle.GetPacket().IsValid())
 	{
-		const TSharedPtr<FOSCBundlePacket>& BundlePacket = StaticCastSharedPtr<FOSCBundlePacket>(Bundle.GetPacket());
+		const TSharedPtr<FOSCBundlePacket>& BundlePacket = StaticCastSharedPtr<FOSCBundlePacket>(InBundle.GetPacket());
 		for (int32 i = 0; i < BundlePacket->GetPackets().Num(); i++)
 		{
 			const TSharedPtr<IOSCPacket>& Packet = BundlePacket->GetPackets()[i];
@@ -206,12 +303,36 @@ TArray<FOSCBundle>& UOSCManager::GetBundlesFromBundle(const FOSCBundle& Bundle,
 	return Bundles;
 }
 
-TArray<FOSCMessage>& UOSCManager::GetMessagesFromBundle(const FOSCBundle& Bundle, TArray<FOSCMessage>& Messages)
+FOSCMessage UOSCManager::GetMessageFromBundle(const FOSCBundle& InBundle, int32 InIndex, bool& bOutSucceeded)
+{
+	if (InBundle.GetPacket().IsValid())
+	{
+		const TSharedPtr<FOSCBundlePacket>& BundlePacket = StaticCastSharedPtr<FOSCBundlePacket>(InBundle.GetPacket());
+		int32 Count = 0;
+		for (const TSharedPtr<IOSCPacket>& Packet : BundlePacket->GetPackets())
+		{
+			if (Packet->IsMessage())
+			{
+				if (InIndex == Count)
+				{
+					bOutSucceeded = true;
+					return FOSCMessage(StaticCastSharedPtr<FOSCMessagePacket>(Packet));
+				}
+				Count++;
+			}
+		}
+	}
+
+	bOutSucceeded = false;
+	return FOSCMessage();
+}
+
+TArray<FOSCMessage> UOSCManager::GetMessagesFromBundle(const FOSCBundle& OutBundle)
 {
-	Messages.Reset();
-	if (Bundle.GetPacket().IsValid())
+	TArray<FOSCMessage> Messages;
+	if (OutBundle.GetPacket().IsValid())
 	{
-		const TSharedPtr<FOSCBundlePacket>& BundlePacket = StaticCastSharedPtr<FOSCBundlePacket>(Bundle.GetPacket());
+		const TSharedPtr<FOSCBundlePacket>& BundlePacket = StaticCastSharedPtr<FOSCBundlePacket>(OutBundle.GetPacket());
 		for (int32 i = 0; i < BundlePacket->GetPackets().Num(); i++)
 		{
 			const TSharedPtr<IOSCPacket>& Packet = BundlePacket->GetPackets()[i];
@@ -225,76 +346,128 @@ TArray<FOSCMessage>& UOSCManager::GetMessagesFromBundle(const FOSCBundle& Bundle
 	return Messages;
 }
 
-void UOSCManager::GetFloat(const FOSCMessage& Message, const int32 Index, float& Value)
+bool UOSCManager::GetAddress(const FOSCMessage& InMessage, const int32 InIndex, FOSCAddress& OutValue)
+{
+	if (const FOSCType* OSCType = OSC::GetOSCTypeAtIndex(InMessage, InIndex))
+	{
+		if (OSCType->IsString())
+		{
+			OutValue = FOSCAddress(OSCType->GetString());
+			return OutValue.IsValidPath();
+		}
+		OSC_LOG_INVALID_TYPE_AT_INDEX("String (OSCAddress)", InIndex, InMessage);
+	}
+
+	OutValue = FOSCAddress();
+	return false;
+}
+
+void UOSCManager::GetAllAddresses(const FOSCMessage& InMessage, TArray<FOSCAddress>& OutValues)
+{
+	if (InMessage.GetPacket().IsValid())
+	{
+		const TSharedPtr<FOSCMessagePacket>& MessagePacket = StaticCastSharedPtr<FOSCMessagePacket>(InMessage.GetPacket());
+		const TArray<FOSCType>& Args = MessagePacket->GetArguments();
+		for (int32 i = 0; i < Args.Num(); i++)
+		{
+			const FOSCType& OSCType = Args[i];
+			if (OSCType.IsString())
+			{
+				FOSCAddress AddressToAdd = FOSCAddress(OSCType.GetString());
+				if (AddressToAdd.IsValidPath())
+				{
+					OutValues.Add(MoveTemp(AddressToAdd));
+				}
+			}
+		}
+	}
+}
+
+bool UOSCManager::GetFloat(const FOSCMessage& InMessage, const int32 InIndex, float& OutValue)
 {
-	if (const FOSCType* OSCType = GetOSCTypeAtIndex(Message, Index))
+	OutValue = 0.0f;
+	if (const FOSCType* OSCType = OSC::GetOSCTypeAtIndex(InMessage, InIndex))
 	{
 		if (OSCType->IsFloat())
 		{
-			Value = OSCType->GetFloat();
+			OutValue = OSCType->GetFloat();
+			return true;
 		}
+		OSC_LOG_INVALID_TYPE_AT_INDEX("Float", InIndex, InMessage);
 	}
+
+	return false;
 }
 
-void UOSCManager::GetAllFloats(const FOSCMessage& Message, TArray<float>& Values)
+void UOSCManager::GetAllFloats(const FOSCMessage& InMessage, TArray<float>& OutValues)
 {
-	if (Message.GetPacket().IsValid())
+	if (InMessage.GetPacket().IsValid())
 	{
-		const TSharedPtr<FOSCMessagePacket>& MessagePacket = StaticCastSharedPtr<FOSCMessagePacket>(Message.GetPacket());
+		const TSharedPtr<FOSCMessagePacket>& MessagePacket = StaticCastSharedPtr<FOSCMessagePacket>(InMessage.GetPacket());
 		const TArray<FOSCType>& Args = MessagePacket->GetArguments();
 		for (int32 i = 0; i < Args.Num(); i++)
 		{
 			const FOSCType& OSCType = Args[i];
 			if (OSCType.IsFloat())
 			{
-				Values.Add(OSCType.GetFloat());
+				OutValues.Add(OSCType.GetFloat());
 			}
 		}
 	}
 }
 
-void UOSCManager::GetInt32(const FOSCMessage& Message, const int32 Index, int32& Value)
+bool UOSCManager::GetInt32(const FOSCMessage& InMessage, const int32 InIndex, int32& OutValue)
 {
-	if (const FOSCType* OSCType = GetOSCTypeAtIndex(Message, Index))
+	OutValue = 0;
+	if (const FOSCType* OSCType = OSC::GetOSCTypeAtIndex(InMessage, InIndex))
 	{
 		if (OSCType->IsInt32())
 		{
-			Value = OSCType->GetInt32();
+			OutValue = OSCType->GetInt32();
+			return true;
 		}
+		OSC_LOG_INVALID_TYPE_AT_INDEX("Int32", InIndex, InMessage);
 	}
+
+	return false;
 }
 
-void UOSCManager::GetAllInt32s(const FOSCMessage& Message, TArray<int32>& Values)
+void UOSCManager::GetAllInt32s(const FOSCMessage& InMessage, TArray<int32>& OutValues)
 {
-	if (Message.GetPacket().IsValid())
+	if (InMessage.GetPacket().IsValid())
 	{
-		const TSharedPtr<FOSCMessagePacket>& MessagePacket = StaticCastSharedPtr<FOSCMessagePacket>(Message.GetPacket());
+		const TSharedPtr<FOSCMessagePacket>& MessagePacket = StaticCastSharedPtr<FOSCMessagePacket>(InMessage.GetPacket());
 		const TArray<FOSCType>& Args = MessagePacket->GetArguments();
 		for (int32 i = 0; i < Args.Num(); i++)
 		{
 			const FOSCType& OSCType = Args[i];
 			if (OSCType.IsInt32())
 			{
-				Values.Add(OSCType.GetInt32());
+				OutValues.Add(OSCType.GetInt32());
 			}
 		}
 	}
 }
 
-void UOSCManager::GetInt64(const FOSCMessage& Message, const int32 Index, int64& Value)
+bool UOSCManager::GetInt64(const FOSCMessage& InMessage, const int32 InIndex, int64& OutValue)
 {
-	if (const FOSCType* OSCType = GetOSCTypeAtIndex(Message, Index))
+	OutValue = 0l;
+	if (const FOSCType* OSCType = OSC::GetOSCTypeAtIndex(InMessage, InIndex))
 	{
 		if (OSCType->IsInt64())
 		{
-			Value = OSCType->GetInt64();
+			OutValue = OSCType->GetInt64();
+			return true;
 		}
+		OSC_LOG_INVALID_TYPE_AT_INDEX("Int64", InIndex, InMessage);
 	}
+
+	return false;
 }
 
-void UOSCManager::GetAllInt64s(const FOSCMessage& Message, TArray<int64>& Values)
+void UOSCManager::GetAllInt64s(const FOSCMessage& InMessage, TArray<int64>& OutValues)
 {
-	const TSharedPtr<FOSCMessagePacket>& MessagePacket = StaticCastSharedPtr<FOSCMessagePacket>(Message.GetPacket());
+	const TSharedPtr<FOSCMessagePacket>& MessagePacket = StaticCastSharedPtr<FOSCMessagePacket>(InMessage.GetPacket());
 	if (MessagePacket.IsValid())
 	{
 		const TArray<FOSCType>& Args = MessagePacket->GetArguments();
@@ -303,26 +476,31 @@ void UOSCManager::GetAllInt64s(const FOSCMessage& Message, TArray<int64>& Values
 			const FOSCType& OSCType = Args[i];
 			if (OSCType.IsInt64())
 			{
-				Values.Add(OSCType.GetInt64());
+				OutValues.Add(OSCType.GetInt64());
 			}
 		}
 	}
 }
 
-void UOSCManager::GetString(const FOSCMessage& Message, const int32 Index, FString& Value)
+bool UOSCManager::GetString(const FOSCMessage& InMessage, const int32 InIndex, FString& OutValue)
 {
-	if (const FOSCType* OSCType = GetOSCTypeAtIndex(Message, Index))
+	if (const FOSCType* OSCType = OSC::GetOSCTypeAtIndex(InMessage, InIndex))
 	{
 		if (OSCType->IsString())
 		{
-			Value = OSCType->GetString();
+			OutValue = OSCType->GetString();
+			return true;
 		}
+		OSC_LOG_INVALID_TYPE_AT_INDEX("String", InIndex, InMessage);
 	}
+
+	OutValue.Reset();
+	return false;
 }
 
-void UOSCManager::GetAllStrings(const FOSCMessage& Message, TArray<FString>& Values)
+void UOSCManager::GetAllStrings(const FOSCMessage& InMessage, TArray<FString>& OutValues)
 {
-	const TSharedPtr<FOSCMessagePacket>& MessagePacket = StaticCastSharedPtr<FOSCMessagePacket>(Message.GetPacket());
+	const TSharedPtr<FOSCMessagePacket>& MessagePacket = StaticCastSharedPtr<FOSCMessagePacket>(InMessage.GetPacket());
 	if (MessagePacket.IsValid())
 	{
 		const TArray<FOSCType>& Args = MessagePacket->GetArguments();
@@ -331,26 +509,31 @@ void UOSCManager::GetAllStrings(const FOSCMessage& Message, TArray<FString>& Val
 			const FOSCType& OSCType = Args[i];
 			if (OSCType.IsString())
 			{
-				Values.Add(OSCType.GetString());
+				OutValues.Add(OSCType.GetString());
 			}
 		}
 	}
 }
 
-void UOSCManager::GetBool(const FOSCMessage& Message, const int32 Index, bool& Value)
+bool UOSCManager::GetBool(const FOSCMessage& InMessage, const int32 InIndex, bool& OutValue)
 {
-	if (const FOSCType* OSCType = GetOSCTypeAtIndex(Message, Index))
+	OutValue = false;
+	if (const FOSCType* OSCType = OSC::GetOSCTypeAtIndex(InMessage, InIndex))
 	{
 		if (OSCType->IsBool())
 		{
-			Value = OSCType->GetBool();
+			OutValue = OSCType->GetBool();
+			return true;
 		}
+		OSC_LOG_INVALID_TYPE_AT_INDEX("Bool", InIndex, InMessage);
 	}
+
+	return false;
 }
 
-void UOSCManager::GetAllBools(const FOSCMessage& Message, TArray<bool>& Values)
+void UOSCManager::GetAllBools(const FOSCMessage& InMessage, TArray<bool>& OutValues)
 {
-	const TSharedPtr<FOSCMessagePacket>& MessagePacket = StaticCastSharedPtr<FOSCMessagePacket>(Message.GetPacket());
+	const TSharedPtr<FOSCMessagePacket>& MessagePacket = StaticCastSharedPtr<FOSCMessagePacket>(InMessage.GetPacket());
 	if (MessagePacket.IsValid())
 	{
 		const TArray<FOSCType>& Args = MessagePacket->GetArguments();
@@ -359,58 +542,116 @@ void UOSCManager::GetAllBools(const FOSCMessage& Message, TArray<bool>& Values)
 			const FOSCType& OSCType = Args[i];
 			if (OSCType.IsBool())
 			{
-				Values.Add(OSCType.GetBool());
+				OutValues.Add(OSCType.GetBool());
 			}
 		}
 	}
 }
 
-void UOSCManager::GetBlob(const FOSCMessage& Message, const int32 Index, TArray<uint8>& Value)
+bool UOSCManager::GetBlob(const FOSCMessage& InMessage, const int32 InIndex, TArray<uint8>& OutValue)
 {
-	if (const FOSCType* OSCType = GetOSCTypeAtIndex(Message, Index))
+	OutValue.Reset();
+	if (const FOSCType* OSCType = OSC::GetOSCTypeAtIndex(InMessage, InIndex))
 	{
 		if (OSCType->IsBlob())
 		{
-			Value = OSCType->GetBlob();
+			OutValue = OSCType->GetBlob();
+			return true;
 		}
+		OSC_LOG_INVALID_TYPE_AT_INDEX("Blob", InIndex, InMessage);
+	}
+
+	return false;
+}
+
+bool UOSCManager::OSCAddressIsValidPath(const FOSCAddress& InAddress)
+{
+	return InAddress.IsValidPath();
+}
+
+bool UOSCManager::OSCAddressIsValidPattern(const FOSCAddress& InAddress)
+{
+	return InAddress.IsValidPattern();
+}
+
+FOSCAddress UOSCManager::ConvertStringToOSCAddress(const FString& InString)
+{
+	return FOSCAddress(InString);
+}
+
+UObject* UOSCManager::FindObjectAtOSCAddress(const FOSCAddress& InAddress)
+{
+	FSoftObjectPath Path(ObjectPathFromOSCAddress(InAddress));
+	if (Path.IsValid())
+	{
+		return Path.TryLoad();
 	}
+
+	UE_LOG(LogOSC, Verbose, TEXT("Failed to load object from OSCAddress '%s'"), *InAddress.GetFullPath());
+	return nullptr;
 }
 
-bool UOSCManager::OSCAddressIsValidPath(const FOSCAddress& Address)
+FOSCAddress UOSCManager::OSCAddressFromObjectPath(UObject* InObject)
 {
-	return Address.IsValidPath();
+	const FString Path = FPaths::ChangeExtension(InObject->GetPathName(), FString());
+	return FOSCAddress(Path);
+}
+
+FOSCAddress UOSCManager::OSCAddressFromObjectPathString(const FString& InPathName)
+{
+	TArray<FString> PartArray;
+	InPathName.ParseIntoArray(PartArray, TEXT("\'"));
+
+	// Type declaration at beginning of path. Assumed to be in the form <SomeTypeContainer1'/Container2/ObjectName.ObjectName'>
+	if (PartArray.Num() > 1)
+	{
+		const FString NoExtPath = FPaths::SetExtension(PartArray[1], TEXT(""));
+		return FOSCAddress(NoExtPath);
+	}
+
+	// No type declaration at beginning of path. Assumed to be in the form <Container1/Container2/ObjectName.ObjectName>
+	if (PartArray.Num() > 0)
+	{
+		const FString NoExtPath = FPaths::SetExtension(PartArray[0], TEXT(""));
+		return FOSCAddress(NoExtPath);
+	}
+
+	// Invalid address
+	return FOSCAddress();
 }
 
-bool UOSCManager::OSCAddressIsValidPattern(const FOSCAddress& Address)
+FString UOSCManager::ObjectPathFromOSCAddress(const FOSCAddress& InAddress)
 {
-	return Address.IsValidPattern();
+	const FString Path = InAddress.GetFullPath() + TEXT(".") + InAddress.GetMethod();
+	return Path;
 }
 
-FOSCAddress UOSCManager::ConvertStringToOSCAddress(const FString& String)
+FOSCAddress& UOSCManager::OSCAddressPushContainer(FOSCAddress& OutAddress, const FString& InToAppend)
 {
-	return FOSCAddress(String);
+	OutAddress.PushContainer(InToAppend);
+	return OutAddress;
 }
 
-FOSCAddress& UOSCManager::OSCAddressPushContainer(FOSCAddress& Address, const FString& ToAppend)
+FOSCAddress& UOSCManager::OSCAddressPushContainers(FOSCAddress& OutAddress, const TArray<FString>& InToAppend)
 {
-	Address.PushContainer(ToAppend);
-	return Address;
+	OutAddress.PushContainers(InToAppend);
+	return OutAddress;
 }
 
-FOSCAddress& UOSCManager::OSCAddressPushContainers(FOSCAddress& Address, const TArray<FString>& ToAppend)
+FString UOSCManager::OSCAddressPopContainer(FOSCAddress& OutAddress)
 {
-	Address.PushContainers(ToAppend);
-	return Address;
+	return OutAddress.PopContainer();
 }
 
-FString UOSCManager::OSCAddressPopContainer(FOSCAddress& Address)
+TArray<FString> UOSCManager::OSCAddressPopContainers(FOSCAddress& OutAddress, int32 InNumContainers)
 {
-	return Address.PopContainer();
+	return OutAddress.PopContainers(InNumContainers);
 }
 
-TArray<FString> UOSCManager::OSCAddressPopContainers(FOSCAddress& Address, int32 InNumContainers)
+FOSCAddress& UOSCManager::OSCAddressRemoveContainers(FOSCAddress& OutAddress, int32 InIndex, int32 InCount)
 {
-	return Address.PopContainers(InNumContainers);
+	OutAddress.RemoveContainers(InIndex, InCount);
+	return OutAddress;
 }
 
 bool UOSCManager::OSCAddressPathMatchesPattern(const FOSCAddress& InPattern, const FOSCAddress& InPath)
@@ -418,52 +659,52 @@ bool UOSCManager::OSCAddressPathMatchesPattern(const FOSCAddress& InPattern, con
 	return InPattern.Matches(InPath);
 }
 
-FOSCAddress UOSCManager::GetOSCMessageAddress(const FOSCMessage& Message)
+FOSCAddress UOSCManager::GetOSCMessageAddress(const FOSCMessage& InMessage)
 {
-	return Message.GetAddress();
+	return InMessage.GetAddress();
 }
 
-FOSCMessage& UOSCManager::SetOSCMessageAddress(FOSCMessage& Message, const FOSCAddress& Address)
+FOSCMessage& UOSCManager::SetOSCMessageAddress(FOSCMessage& OutMessage, const FOSCAddress& InAddress)
 {
-	Message.SetAddress(Address);
-	return Message;
+	OutMessage.SetAddress(InAddress);
+	return OutMessage;
 }
 
-FString UOSCManager::GetOSCAddressContainer(const FOSCAddress& Address, int32 Index)
+FString UOSCManager::GetOSCAddressContainer(const FOSCAddress& InAddress, int32 InIndex)
 {
-	return Address.GetContainer(Index);
+	return InAddress.GetContainer(InIndex);
 }
 
-TArray<FString> UOSCManager::GetOSCAddressContainers(const FOSCAddress& Address)
+TArray<FString> UOSCManager::GetOSCAddressContainers(const FOSCAddress& InAddress)
 {
 	TArray<FString> Containers;
-	Address.GetContainers(Containers);
+	InAddress.GetContainers(Containers);
 	return Containers;
 }
 
-FString UOSCManager::GetOSCAddressContainerPath(const FOSCAddress& Address)
+FString UOSCManager::GetOSCAddressContainerPath(const FOSCAddress& InAddress)
 {
-	return Address.GetContainerPath();
+	return InAddress.GetContainerPath();
 }
 
-FString UOSCManager::GetOSCAddressFullPath(const FOSCAddress& Address)
+FString UOSCManager::GetOSCAddressFullPath(const FOSCAddress& InAddress)
 {
-	return Address.GetFullPath();
+	return InAddress.GetFullPath();
 }
 
-FString UOSCManager::GetOSCAddressMethod(const FOSCAddress& Address)
+FString UOSCManager::GetOSCAddressMethod(const FOSCAddress& InAddress)
 {
-	return Address.GetMethod();
+	return InAddress.GetMethod();
 }
 
-FOSCAddress& UOSCManager::ClearOSCAddressContainers(FOSCAddress& Address)
+FOSCAddress& UOSCManager::ClearOSCAddressContainers(FOSCAddress& OutAddress)
 {
-	Address.ClearContainers();
-	return Address;
+	OutAddress.ClearContainers();
+	return OutAddress;
 }
 
-FOSCAddress& UOSCManager::SetOSCAddressMethod(FOSCAddress& Address, const FString& Method)
+FOSCAddress& UOSCManager::SetOSCAddressMethod(FOSCAddress& OutAddress, const FString& InMethod)
 {
-	Address.SetMethod(Method);
-	return Address;
+	OutAddress.SetMethod(InMethod);
+	return OutAddress;
 }
\ No newline at end of file
diff --git a/Engine/Plugins/Runtime/OSC/Source/OSC/Private/OSCMessage.cpp b/Engine/Plugins/Runtime/OSC/Source/OSC/Private/OSCMessage.cpp
--- a/Engine/Plugins/Runtime/OSC/Source/OSC/Private/OSCMessage.cpp
+++ b/Engine/Plugins/Runtime/OSC/Source/OSC/Private/OSCMessage.cpp
@@ -6,7 +6,7 @@
 
 
 FOSCMessage::FOSCMessage()
-	: Packet(MakeShareable(new FOSCMessagePacket()))
+	: Packet(MakeShared<FOSCMessagePacket>())
 {
 }
 
diff --git a/Engine/Plugins/Runtime/OSC/Source/OSC/Private/OSCMessagePacket.cpp b/Engine/Plugins/Runtime/OSC/Source/OSC/Private/OSCMessagePacket.cpp
--- a/Engine/Plugins/Runtime/OSC/Source/OSC/Private/OSCMessagePacket.cpp
+++ b/Engine/Plugins/Runtime/OSC/Source/OSC/Private/OSCMessagePacket.cpp
@@ -117,7 +117,13 @@ void FOSCMessagePacket::ReadData(FOSCStream& Stream)
 
 	// Read string of tags
 	const FString StreamString = Stream.ReadString();
+
 	const TArray<TCHAR>& TagTypes = StreamString.GetCharArray();
+	if(TagTypes.Num() == 0)
+	{
+		UE_LOG(LogOSC, Error, TEXT("Failed to read message packet with address '%s' from stream: Invalid (Empty) Type Tag"), *Address.GetFullPath());
+		return;
+	}
 
 	// Skip the first argument which is ','
 	for (int32 i = 1; i < TagTypes.Num(); i++)
@@ -165,8 +171,9 @@ void FOSCMessagePacket::ReadData(FOSCStream& Stream)
 			Arguments.Add(FOSCType(FColor(Stream.ReadInt32())));
 			break;
 		case EOSCTypeTag::TERMINATE:
-			Stream.ReadChar();
-			break;
+			// Return on first terminate found. FString GetCharArray can return
+			// an array with multiple terminators.
+			return;
 
 		default:
 			// Argument is not supported 
diff --git a/Engine/Plugins/Runtime/OSC/Source/OSC/Private/OSCServer.cpp b/Engine/Plugins/Runtime/OSC/Source/OSC/Private/OSCServer.cpp
--- a/Engine/Plugins/Runtime/OSC/Source/OSC/Private/OSCServer.cpp
+++ b/Engine/Plugins/Runtime/OSC/Source/OSC/Private/OSCServer.cpp
@@ -94,6 +94,26 @@ void UOSCServer::ClearWhitelistedClients()
 	ServerProxy->ClearWhitelistedClients();
 }
 
+FString UOSCServer::GetIpAddress(bool bIncludePort) const
+{
+	check(ServerProxy.IsValid());
+
+	FString Address = ServerProxy->GetIpAddress();
+	if (bIncludePort)
+	{
+		Address += TEXT(":");
+		Address.AppendInt(ServerProxy->GetPort());
+	}
+
+	return Address;
+}
+
+int32 UOSCServer::GetPort() const
+{
+	check(ServerProxy.IsValid());
+	return ServerProxy->GetPort();
+}
+
 TSet<FString> UOSCServer::GetWhitelistedClients() const
 {
 	check(ServerProxy.IsValid());
diff --git a/Engine/Plugins/Runtime/OSC/Source/OSC/Private/OSCServerProxy.cpp b/Engine/Plugins/Runtime/OSC/Source/OSC/Private/OSCServerProxy.cpp
--- a/Engine/Plugins/Runtime/OSC/Source/OSC/Private/OSCServerProxy.cpp
+++ b/Engine/Plugins/Runtime/OSC/Source/OSC/Private/OSCServerProxy.cpp
@@ -47,6 +47,16 @@ void FOSCServerProxy::OnPacketReceived(const FArrayReaderPtr& Data, const FIPv4E
 	}, GET_STATID(STAT_OSCServerOnPacketReceived), nullptr, ENamedThreads::GameThread);
 }
 
+FString FOSCServerProxy::GetIpAddress() const
+{
+	return ReceiveIPAddress.ToString();
+}
+
+int32 FOSCServerProxy::GetPort() const
+{
+	return Port;
+}
+
 bool FOSCServerProxy::GetMulticastLoopback() const
 {
 	return bMulticastLoopback;
diff --git a/Engine/Plugins/Runtime/OSC/Source/OSC/Private/OSCServerProxy.h b/Engine/Plugins/Runtime/OSC/Source/OSC/Private/OSCServerProxy.h
--- a/Engine/Plugins/Runtime/OSC/Source/OSC/Private/OSCServerProxy.h
+++ b/Engine/Plugins/Runtime/OSC/Source/OSC/Private/OSCServerProxy.h
@@ -18,7 +18,8 @@ public:
 	virtual ~FOSCServerProxy() { }
 
 	bool GetMulticastLoopback() const override;
-
+	FString GetIpAddress() const override;
+	int32 GetPort() const override;
 	bool IsActive() const override;
 
 	void Listen(const FString& ServerName) override;
diff --git a/Engine/Plugins/Runtime/OSC/Source/OSC/Private/OSCStream.cpp b/Engine/Plugins/Runtime/OSC/Source/OSC/Private/OSCStream.cpp
--- a/Engine/Plugins/Runtime/OSC/Source/OSC/Private/OSCStream.cpp
+++ b/Engine/Plugins/Runtime/OSC/Source/OSC/Private/OSCStream.cpp
@@ -4,23 +4,27 @@
 #include "OSCLog.h"
 
 
+namespace OSC
+{
+	const int32 DefaultWriteStreamSize = 1024;
+} // namespace OSC
+
 FOSCStream::FOSCStream()
 	: Position(0)
+	, bIsReadStream(false)
 {
+	// Set default size to reasonable size to
+	// avoid thrashing with allocations while writing
+	Data.Reserve(OSC::DefaultWriteStreamSize);
 }
 
 FOSCStream::FOSCStream(const uint8* InData, int32 InSize)
 	: Data(InData, InSize)
 	, Position(0)
+	, bIsReadStream(true)
 {
 }
 
-FOSCStream::FOSCStream(int32 InSize)
-	: Position(0)
-{
-	Data.AddZeroed(InSize);
-}
-
 const uint8* FOSCStream::GetData() const
 {
 	return Data.GetData();
@@ -43,6 +47,7 @@ int32 FOSCStream::GetPosition() const
 
 void FOSCStream::SetPosition(int32 InPosition)
 {
+	check(InPosition <= Data.Num());
 	Position = InPosition;
 }
 
@@ -346,42 +351,46 @@ void FOSCStream::WriteFloat(float Value)
 FString FOSCStream::ReadString()
 {
 	const int32 DataSize = Data.Num();
-	const int32 InitPosition = Position;
-	check(Position < DataSize);
-
-	const ANSICHAR* StrStart = (ANSICHAR*)(&Data[InitPosition]);
-
-	for (; Position < DataSize; Position++)
+	if (HasReachedEnd() || DataSize == 0)
 	{
-		if (Data[Position] == '\0')
-		{
-			break;
-		}
+		return FString();
 	}
 
-	if (Position == DataSize)
+	// Cache init position index and push along until either at end or character read is null terminator.
+	const int32 InitPosition = GetPosition();
+	while (!HasReachedEnd() && ReadChar() != '\0');
+
+	if (HasReachedEnd() && Data.Last() != '\0')
 	{
-		UE_LOG(LogOSC, Error, TEXT("Invalid string: Terminator not found"));
+		UE_LOG(LogOSC, Error, TEXT("Invalid string when reading OSCStream: Null terminator '\0' not found"));
 		return FString();
 	}
 
-	// Note end for string copy, increment to next read
-	// location, and pad position.
+	// Cache end for string copy, increment to next read
+	// location, and consume pad until next 4-byte boundary.
 	const int32 EndPosition = Position;
-	Position++;
-	Position = ((Position + 3) / 4) * 4;
 
-	return FString(EndPosition - InitPosition, StrStart);
+	const int32 Count = EndPosition - InitPosition;
+	check(Count > 0);
+
+	const int32 UnboundByteCount = Count % 4;
+	if (UnboundByteCount != 0)
+	{
+		Position += 4 - UnboundByteCount;
+	}
+
+	return FString(Count, (ANSICHAR*)(&Data[InitPosition]));
 }
 
 void FOSCStream::WriteString(const FString& InString)
 {
 	const TArray<TCHAR>& CharArr = InString.GetCharArray();
-	const int32 Count = CharArr.Num();
 
+	int32 Count = CharArr.Num();
 	if (Count == 0)
 	{
 		WriteChar('\0');
+		Count++;
 	}
 	else
 	{
@@ -391,11 +400,17 @@ void FOSCStream::WriteString(const FString& InString)
 		}
 	}
 
-	// Increment & pad string with null terminator
-	const int32 NumPaddingZeros = ((Count + 3) / 4) * 4;
-	for (int32 i = 0; i < NumPaddingZeros - Count; i++)
+	// Increment & pad string with null terminator (String must
+	// be packed into multiple of 4 bytes)
+	check(Count > 0);
+	const int32 UnboundByteCount = Count % 4;
+	if (UnboundByteCount != 0)
 	{
-		WriteChar('\0');
+		const int32 NumPaddingZeros = 4 - UnboundByteCount;
+		for (int32 i = 0; i < NumPaddingZeros; i++)
+		{
+			WriteChar('\0');
+		}
 	}
 }
 
@@ -425,6 +440,7 @@ void FOSCStream::WriteBlob(TArray<uint8>& Blob)
 
 int32 FOSCStream::Read(uint8* Buffer, int32 InSize)
 {
+	check(bIsReadStream);
 	const int32 DataSize = Data.Num();
 	if (InSize == 0 || Position >= DataSize)
 	{
@@ -443,13 +459,28 @@ int32 FOSCStream::Read(uint8* Buffer, int32 InSize)
 
 int32 FOSCStream::Write(const uint8* InBuffer, int32 InSize)
 {
-	if (InSize > 0)
+	check(!bIsReadStream);
+	if (InSize <= 0)
 	{
-		FMemory::Memcpy(&Data[Position], InBuffer, InSize);
-		Position += InSize;
-		return InSize;
+		return 0;
 	}
 
-	return 0;
+	if (Position < Data.Num())
+	{
+		int32 Slack = Data.Num() - Position;
+		if (InSize - Slack > 0)
+		{
+			Data.AddUninitialized(InSize - Slack);
+		}
+	}
+	else
+	{
+		check(Position == Data.Num());
+		Data.AddUninitialized(InSize);
+	}
+
+	FMemory::Memcpy(&Data[Position], InBuffer, InSize);
+	Position += InSize;
+	return InSize;
 }
 
diff --git a/Engine/Plugins/Runtime/OSC/Source/OSC/Public/OSCAddress.h b/Engine/Plugins/Runtime/OSC/Source/OSC/Public/OSCAddress.h
--- a/Engine/Plugins/Runtime/OSC/Source/OSC/Public/OSCAddress.h
+++ b/Engine/Plugins/Runtime/OSC/Source/OSC/Public/OSCAddress.h
@@ -64,6 +64,9 @@ public:
 	/** Pops containers off end of address' ordered array of containers */
 	TArray<FString> PopContainers(int32 InNumContainers);
 
+	/** Removes containers from container array at index until count */
+	void RemoveContainers(int32 InIndex, int32 InCount);
+
 	/** Clears ordered array of containers */
 	void ClearContainers();
 
@@ -87,7 +90,48 @@ public:
 
 	bool operator== (const FOSCAddress& InAddress) const
 	{
-		return Hash == InAddress.Hash;
+		if (Hash != InAddress.Hash)
+		{
+			return false;
+		}
+
+		if (InAddress.Containers.Num() != Containers.Num())
+		{
+			return false;
+		}
+
+		if (InAddress.Method != Method)
+		{
+			return false;
+		}
+
+		for (int32 i = 0; i < Containers.Num(); ++i)
+		{
+			if (InAddress.Containers[i] != Containers[i])
+			{
+				return false;
+			}
+		}
+
+		return true;
+	}
+
+	FOSCAddress& operator/= (const FOSCAddress& InAddress)
+	{
+		Containers.Add(Method);
+		Containers.Append(InAddress.Containers);
+
+		Method = InAddress.Method;
+
+		CacheAggregates();
+		return *this;
+	}
+
+	FOSCAddress operator/ (const FOSCAddress& InAddress) const
+	{
+		FOSCAddress ToReturn = *this;
+		ToReturn /= InAddress;
+		return ToReturn;
 	}
 
 	friend uint32 GetTypeHash(const FOSCAddress& InAddress)
diff --git a/Engine/Plugins/Runtime/OSC/Source/OSC/Public/OSCClient.h b/Engine/Plugins/Runtime/OSC/Source/OSC/Public/OSCClient.h
--- a/Engine/Plugins/Runtime/OSC/Source/OSC/Public/OSCClient.h
+++ b/Engine/Plugins/Runtime/OSC/Source/OSC/Public/OSCClient.h
@@ -20,6 +20,8 @@ public:
 	virtual void GetSendIPAddress(FString& InIPAddress, int32& Port) const = 0;
 	virtual bool SetSendIPAddress(const FString& InIPAddress, const int32 Port) = 0;
 
+	virtual bool IsActive() const = 0;
+
 	virtual void SendMessage(FOSCMessage& Message) = 0;
 	virtual void SendBundle(FOSCBundle& Bundle) = 0;
 
@@ -34,6 +36,8 @@ class OSC_API UOSCClient : public UObject
 public:
 	void Connect();
 
+	bool IsActive() const;
+
 	/** Gets the OSC Client IP address and port. */
 	UFUNCTION(BlueprintCallable, Category = "Audio|OSC")
 	void GetSendIPAddress(UPARAM(ref) FString& IPAddress, UPARAM(ref) int32& Port);
diff --git a/Engine/Plugins/Runtime/OSC/Source/OSC/Public/OSCManager.h b/Engine/Plugins/Runtime/OSC/Source/OSC/Public/OSCManager.h
--- a/Engine/Plugins/Runtime/OSC/Source/OSC/Public/OSCManager.h
+++ b/Engine/Plugins/Runtime/OSC/Source/OSC/Public/OSCManager.h
@@ -2,12 +2,20 @@
 #pragma once
 
 #include "CoreMinimal.h"
-#include "UObject/ObjectMacros.h"
 #include "Kismet/BlueprintFunctionLibrary.h"
-#include "OSCServer.h"
-#include "OSCClient.h"
+#include "OSCAddress.h"
+#include "OSCBundle.h"
+#include "OSCMessage.h"
+#include "UObject/ObjectMacros.h"
+
+
 #include "OSCManager.generated.h"
 
+
+// Forward Declarations
+class UOSCServer;
+class UOSCClient;
+
 UCLASS()
 class OSC_API UOSCManager : public UBlueprintFunctionLibrary
 {
@@ -18,15 +26,15 @@ public:
 	// attempts to use LocalHost IP address. If StartListening set,
 	// immediately begins listening on creation.
 	UFUNCTION(BlueprintCallable, Category = "Audio|OSC")
-	static class UOSCServer* CreateOSCServer(FString ReceiveIPAddress, int32 Port, bool bMulticastLoopback, bool bStartListening);
+	static UOSCServer* CreateOSCServer(FString ReceiveIPAddress, int32 Port, bool bMulticastLoopback, bool bStartListening);
 
 	// Creates an OSC Client.  If left empty (or '0'), attempts to use
 	// attempts to use LocalHost IP address.
 	UFUNCTION(BlueprintCallable, Category = "Audio|OSC")
-	static class UOSCClient* CreateOSCClient(FString SendIPAddress, int32 Port);
+	static UOSCClient* CreateOSCClient(FString SendIPAddress, int32 Port);
 
 	/** Adds provided message packet to bundle. */
-	UFUNCTION(BlueprintCallable, Category = "Audio|OSC", meta = (DisplayName = "Add OSC MEssage to Bundle", Keywords = "osc message bundle"))
+	UFUNCTION(BlueprintCallable, Category = "Audio|OSC", meta = (DisplayName = "Add OSC Message to Bundle", Keywords = "osc message bundle"))
 	static UPARAM(DisplayName = "Bundle") FOSCBundle& AddMessageToBundle(const FOSCMessage& Message, UPARAM(ref) FOSCBundle& Bundle);
 
 	/** Adds bundle packet to bundle. */
@@ -35,11 +43,15 @@ public:
 
 	/** Fills array with child bundles found in bundle. */
 	UFUNCTION(BlueprintCallable, Category = "Audio|OSC", meta = (DisplayName = "Get OSC Bundles From Bundle", Keywords = "osc bundle"))
-	static UPARAM(DisplayName = "Bundles") TArray<FOSCBundle>& GetBundlesFromBundle(const FOSCBundle& Bundle, TArray<FOSCBundle>& Bundles);
+	static UPARAM(DisplayName = "Bundles") TArray<FOSCBundle> GetBundlesFromBundle(const FOSCBundle& Bundle);
+
+	/** Returns message found in bundle at ordered index. */
+	UFUNCTION(BlueprintCallable, Category = "Audio|OSC", meta = (DisplayName = "Get OSC Messages From Bundle", Keywords = "osc bundle message"))
+	static UPARAM(DisplayName = "Message") FOSCMessage GetMessageFromBundle(const FOSCBundle& Bundle, int32 Index, bool& bSucceeded);
 
 	/** Fills array with messages found in bundle. */
 	UFUNCTION(BlueprintCallable, Category = "Audio|OSC", meta = (DisplayName = "Get OSC Messages From Bundle", Keywords = "osc bundle message"))
-	static UPARAM(DisplayName = "Messages") TArray<FOSCMessage>& GetMessagesFromBundle(const FOSCBundle& Bundle, TArray<FOSCMessage>& Messages);
+	static UPARAM(DisplayName = "Messages") TArray<FOSCMessage> GetMessagesFromBundle(const FOSCBundle& Bundle);
 
 	/** Clears provided message of all arguments. */
 	UFUNCTION(BlueprintCallable, Category = "Audio|OSC", meta = (DisplayName = "Clear OSC Message", Keywords = "osc message"))
@@ -61,6 +73,10 @@ public:
 	UFUNCTION(BlueprintCallable, Category = "Audio|OSC", meta = (DisplayName = "Add Integer (64-bit) to OSC Message", Keywords = "osc message"))
 	static UPARAM(DisplayName = "Message") FOSCMessage& AddInt64(UPARAM(ref) FOSCMessage& Message, int64 Value);
 
+	/** Adds address (packed as string) value to end of OSCMessage */
+	UFUNCTION(BlueprintCallable, Category = "Audio|OSC", meta = (DisplayName = "Add OSC Address (As String) to OSC Message", Keywords = "osc message"))
+	static UPARAM(DisplayName = "Message") FOSCMessage& AddAddress(UPARAM(ref) FOSCMessage& Message, const FOSCAddress& Value);
+
 	/** Adds string value to end of OSCMessage */
 	UFUNCTION(BlueprintCallable, Category = "Audio|OSC", meta = (DisplayName = "Add String to OSC Message", Keywords = "osc message"))
 	static UPARAM(DisplayName = "Message") FOSCMessage& AddString(UPARAM(ref) FOSCMessage& Message, FString Value);
@@ -73,9 +89,23 @@ public:
 	UFUNCTION(BlueprintCallable, Category = "Audio|OSC", meta = (DisplayName = "Add Bool to OSC Message", Keywords = "osc message"))
 	static UPARAM(DisplayName = "Message") FOSCMessage& AddBool(UPARAM(ref) FOSCMessage& Message, bool Value);
 
-	/** Set Value to float at provided Index in OSCMessage if in bounds and type matches */
+	/** Sets Value to address at provided Index in OSCMessage if in bounds and OSC type matches 'String' (Does NOT return address of message, rather
+	 * string packed in message and casts to OSC address). Returns if string found at index and is valid OSC address path.
+	 */
+	UFUNCTION(BlueprintCallable, Category = "Audio|OSC", meta = (DisplayName = "Get OSC Message Address At Index", Keywords = "osc message"))
+	static UPARAM(DisplayName = "Succeeded") bool GetAddress(const FOSCMessage& Message, const int32 Index, FOSCAddress& Value);
+
+	/** Returns all strings that are valid address paths in order received from OSCMessage (Does NOT include address of message, just
+	 * strings packed in message that are valid paths).
+	 */
+	UFUNCTION(BlueprintCallable, Category = "Audio|OSC", meta = (DisplayName = "Get OSC Message Addresses", Keywords = "osc message"))
+	static void GetAllAddresses(const FOSCMessage& Message, TArray<FOSCAddress>& Values);
+
+	/**
+	 * Set Value to float at provided Index in OSCMessage if in bounds and type matches
+	 */
 	UFUNCTION(BlueprintCallable, Category = "Audio|OSC", meta = (DisplayName = "Get OSC Message Float At Index", Keywords = "osc message"))
-	static void GetFloat(const FOSCMessage& Message, const int32 Index, float& Value);
+	static UPARAM(DisplayName = "Succeeded") bool GetFloat(const FOSCMessage& Message, const int32 Index, float& Value);
 
 	/** Returns all float values in order of received from OSCMessage */
 	UFUNCTION(BlueprintCallable, Category = "Audio|OSC", meta = (DisplayName = "Get OSC Message Floats", Keywords = "osc message"))
@@ -83,7 +113,7 @@ public:
 
 	/** Set Value to integer at provided Index in OSCMessage if in bounds and type matches */
 	UFUNCTION(BlueprintCallable, Category = "Audio|OSC", meta = (DisplayName = "Get OSC Message Integer at Index", Keywords = "osc message"))
-	static void GetInt32(const FOSCMessage& Message, const int32 Index, int32& Value);
+	static UPARAM(DisplayName = "Succeeded") bool GetInt32(const FOSCMessage& Message, const int32 Index, int32& Value);
 
 	/** Returns all integer values in order of received from OSCMessage */
 	UFUNCTION(BlueprintCallable, Category = "Audio|OSC", meta = (DisplayName = "Get OSC Message Integers", Keywords = "osc message"))
@@ -91,7 +121,7 @@ public:
 
 	/** Set Value to Int64 at provided Index in OSCMessage if in bounds and type matches */
 	UFUNCTION(BlueprintCallable, Category = "Audio|OSC", meta = (DisplayName = "Get OSC Message Integer (64-bit) at Index", Keywords = "osc message"))
-	static void GetInt64(const FOSCMessage& Message, const int32 Index, int64& Value);
+	static UPARAM(DisplayName = "Succeeded") bool GetInt64(const FOSCMessage& Message, const int32 Index, int64& Value);
 
 	/** Returns all Int64 values in order of received from OSCMessage */
 	UFUNCTION(BlueprintCallable, Category = "Audio|OSC", meta = (DisplayName = "Get OSC Message Integers (64-bit)", Keywords = "osc message"))
@@ -99,7 +129,7 @@ public:
 
 	/** Set Value to string at provided Index in OSCMessage if in bounds and type matches */
 	UFUNCTION(BlueprintCallable, Category = "Audio|OSC", meta = (DisplayName = "Get OSC Message String at Index", Keywords = "osc message"))
-	static void GetString(const FOSCMessage& Message, const int32 Index, FString& Value);
+	static UPARAM(DisplayName = "Succeeded") bool GetString(const FOSCMessage& Message, const int32 Index, FString& Value);
 
 	/** Returns all string values in order of received from OSCMessage */
 	UFUNCTION(BlueprintCallable, Category = "Audio|OSC", meta = (DisplayName = "Get OSC Message Strings", Keywords = "osc message"))
@@ -107,7 +137,7 @@ public:
 
 	/** Sets Value to boolean at provided Index from OSCMessage if in bounds and type matches */
 	UFUNCTION(BlueprintCallable, Category = "Audio|OSC", meta = (DisplayName = "Get OSC Message Bool At Index", Keywords = "osc message"))
-	static void GetBool(const FOSCMessage& Message, const int32 Index, bool& Value);
+	static UPARAM(DisplayName = "Succeeded") bool GetBool(const FOSCMessage& Message, const int32 Index, bool& Value);
 
 	/** Returns all boolean values in order of received from OSCMessage */
 	UFUNCTION(BlueprintCallable, Category = "Audio|OSC", meta = (DisplayName = "Get OSC Message Bools", Keywords = "osc message"))
@@ -115,17 +145,17 @@ public:
 
 	/** Sets Value to blob at provided Index from OSCMessage if in bounds and type matches */
 	UFUNCTION(BlueprintCallable, Category = "Audio|OSC")
-	static void GetBlob(const FOSCMessage& Message, const int32 Index, TArray<uint8>& Value);
+	static UPARAM(DisplayName = "Succeeded") bool GetBlob(const FOSCMessage& Message, const int32 Index, TArray<uint8>& Value);
 
-	/** Returns whether OSCAddress is valid path */
+	/** Returns whether OSC Address is valid path */
 	UFUNCTION(BlueprintCallable, Category = "Audio|OSC", meta = (DisplayName = "Is OSC Address Valid Path", Keywords = "valid osc address path address"))
 	static bool OSCAddressIsValidPath(const FOSCAddress& Address);
 
-	/** Returns whether OSCAddress is valid pattern to match against */
+	/** Returns whether OSC Address is valid pattern to match against */
 	UFUNCTION(BlueprintCallable, Category = "Audio|OSC", meta = (DisplayName = "Is OSC Address Valid Pattern", Keywords = "valid osc address pattern address"))
 	static bool OSCAddressIsValidPattern(const FOSCAddress& Address);
 
-	/* Converts string to OSCAddress */
+	/* Converts string to OSC Address */
 	UFUNCTION(BlueprintCallable, Category = "Audio|OSC")
 	static UPARAM(DisplayName = "Address") FOSCAddress ConvertStringToOSCAddress(const FString& String);
 
@@ -135,6 +165,28 @@ public:
 	UFUNCTION(BlueprintCallable, Category = "Audio|OSC", meta = (DisplayName = "OSC Address Path Matches Pattern", Keywords = "matches osc address path address"))
 	static UPARAM(DisplayName = "Is Match") bool OSCAddressPathMatchesPattern(const FOSCAddress& Pattern, const FOSCAddress& Path);
 
+	/** Finds an object with the given OSC Address in path form, where containers correspond to path folders and the the address method to the object's name.
+	  * Only supports parent objects.
+	  */
+	UFUNCTION(BlueprintCallable, Category = "Audio|OSC", meta = (DisplayName = "Find Object at OSC Address", Keywords = "osc address path uobject"))
+	static UPARAM(DisplayName = "Object") UObject* FindObjectAtOSCAddress(const FOSCAddress& Address);
+
+	/** Converts object path to OSC Address, converting folders to address containers and the object's name to the address method.
+	 * Only supports parent objects (See UObjectBaseUtility::GetPathName and UObjectBaseUtility::GetFullName).
+	 */
+	UFUNCTION(BlueprintCallable, Category = "Audio|OSC", meta = (DisplayName = "Convert Object Path to OSC Address", Keywords = "osc address path uobject"))
+	static UPARAM(DisplayName = "Address") FOSCAddress OSCAddressFromObjectPath(UObject* Object);
+
+	/** Converts object path string to OSC Address, converting folders to address containers and the object's name to the address method.
+	 * Only supports parent objects (See UObjectBaseUtility::GetPathName and UObjectBaseUtility::GetFullName).
+	 */
+	UFUNCTION(BlueprintCallable, Category = "Audio|OSC", meta = (DisplayName = "Convert Object Path (String) to OSC Address", Keywords = "osc address path uobject"))
+	static UPARAM(DisplayName = "Address") FOSCAddress OSCAddressFromObjectPathString(const FString& PathName);
+
+	/** Converts OSC Address to an object path. */
+	UFUNCTION(BlueprintCallable, Category = "Audio|OSC", meta = (DisplayName = "Convert OSC Address to Object Path", Keywords = "osc address path uobject"))
+	static UPARAM(DisplayName = "Path") FString ObjectPathFromOSCAddress(const FOSCAddress& Address);
+
 	/** Pushes container onto address' ordered array of containers */
 	UFUNCTION(BlueprintCallable, Category = "Audio|OSC", meta = (DisplayName = "Push Container to OSC Address", Keywords = "push osc address container"))
 	static UPARAM(DisplayName = "Address") FOSCAddress& OSCAddressPushContainer(UPARAM(ref) FOSCAddress& Address, const FString& Container);
@@ -151,15 +203,19 @@ public:
 	UFUNCTION(BlueprintCallable, Category = "Audio|OSC", meta = (DisplayName = "Pop Containers from OSC Address", Keywords = "pop osc address container"))
 	static UPARAM(DisplayName = "Containers") TArray<FString> OSCAddressPopContainers(UPARAM(ref) FOSCAddress& Address, int32 NumContainers);
 
-	/* Returns copy of message's OSCAddress */
-	UFUNCTION(BlueprintCallable, Category = "Audio|OSC")
+	/** Remove containers from ordered array of containers at index up to count of containers. */
+	UFUNCTION(BlueprintCallable, Category = "Audio|OSC", meta = (DisplayName = "Remove Containers from OSC Address", Keywords = "remove osc address container"))
+	static UPARAM(DisplayName = "Address") FOSCAddress& OSCAddressRemoveContainers(UPARAM(ref) FOSCAddress& Address, int32 Index, int32 Count);
+
+	/* Returns copy of message's OSC Address */
+	UFUNCTION(BlueprintCallable, Category = "Audio|OSC", meta = (DisplayName = "Get OSC Message Address", Keywords = "osc address"))
 	static UPARAM(DisplayName = "Address") FOSCAddress GetOSCMessageAddress(const FOSCMessage& Message);
 
-	/** Sets the OSCAddress of the provided message */
+	/** Sets the OSC Address of the provided message */
 	UFUNCTION(BlueprintCallable, Category = "Audio|OSC")
 	static UPARAM(DisplayName = "Message") FOSCMessage& SetOSCMessageAddress(UPARAM(ref) FOSCMessage& Message, const FOSCAddress& Address);
 
-	/** Returns the OSCAddress container at the provided 'Index.' Returns empty string if index is out-of-bounds. */
+	/** Returns the OSC Address container at the provided 'Index.' Returns empty string if index is out-of-bounds. */
 	UFUNCTION(BlueprintCallable, Category = "Audio|OSC", meta = (DisplayName = "Get OSC Address Container At Index", Keywords = "osc address container path"))
 	static UPARAM(DisplayName = "Container") FString GetOSCAddressContainer(const FOSCAddress& Address, const int32 Index);
 
@@ -172,19 +228,19 @@ public:
 	static UPARAM(DisplayName = "Path") FString GetOSCAddressContainerPath(const FOSCAddress& Address);
 
 	/** Returns full path of OSC address in the form '/Container1/Container2' */
-	UFUNCTION(BlueprintCallable, Category = "Audio|OSC", meta = (DisplayName = "Convert OSC Address To String", Keywords = "osc address path"))
+	UFUNCTION(BlueprintCallable, Category = "Audio|OSC", meta = (DisplayName = "Convert OSC Address To String"))
 	static UPARAM(DisplayName = "Full Path") FString GetOSCAddressFullPath(const FOSCAddress& Address);
 
-	/** Returns method name of OSCAddress provided */
-	UFUNCTION(BlueprintCallable, Category = "Audio|OSC")
+	/** Returns method name of OSC Address provided */
+	UFUNCTION(BlueprintCallable, Category = "Audio|OSC", meta = (DisplayName = "Get OSC Address Method"))
 	static UPARAM(DisplayName = "Method") FString GetOSCAddressMethod(const FOSCAddress& Address);
 
-	/** Clears containers of OSCAddress provided */
+	/** Clears containers of OSC Address provided */
 	UFUNCTION(BlueprintCallable, Category = "Audio|OSC", meta = (DisplayName = "Clear OSC Address Containers", Keywords = "osc address clear"))
-	static UPARAM(DisplayName = "Address") FOSCAddress& ClearOSCAddressContainers(FOSCAddress& Address);
+	static UPARAM(DisplayName = "Address") FOSCAddress& ClearOSCAddressContainers(UPARAM(ref) FOSCAddress& Address);
 
-	/** Sets the method name of the OSCAddress provided */
+	/** Sets the method name of the OSC Address provided */
 	UFUNCTION(BlueprintCallable, Category = "Audio|OSC", meta = (DisplayName = "Set OSC Address Method", Keywords = "osc method"))
-	static UPARAM(DisplayName = "Address") FOSCAddress& SetOSCAddressMethod(FOSCAddress& Address, const FString& Method);
+	static UPARAM(DisplayName = "Address") FOSCAddress& SetOSCAddressMethod(UPARAM(ref) FOSCAddress& Address, const FString& Method);
 };
 
diff --git a/Engine/Plugins/Runtime/OSC/Source/OSC/Public/OSCServer.h b/Engine/Plugins/Runtime/OSC/Source/OSC/Public/OSCServer.h
--- a/Engine/Plugins/Runtime/OSC/Source/OSC/Public/OSCServer.h
+++ b/Engine/Plugins/Runtime/OSC/Source/OSC/Public/OSCServer.h
@@ -33,6 +33,8 @@ class OSC_API IOSCServerProxy
 public:
 	virtual ~IOSCServerProxy() { }
 
+	virtual FString GetIpAddress() const = 0;
+	virtual int32 GetPort() const = 0;
 	virtual bool GetMulticastLoopback() const = 0;
 	virtual bool IsActive() const = 0;
 	virtual void Listen(const FString& ServerName) = 0;
@@ -104,6 +106,14 @@ public:
 	UFUNCTION(BlueprintCallable, Category = "Audio|OSC")
 	void ClearWhitelistedClients();
 
+	/** Returns the IP for the server if connected as a string. */
+	UFUNCTION(BlueprintCallable, Category = "Audio|OSC")
+	FString GetIpAddress(bool bIncludePort) const;
+
+	/** Returns the port for the server if connected. */
+	UFUNCTION(BlueprintCallable, Category = "Audio|OSC")
+	int32 GetPort() const;
+
 	/** Returns set of whitelisted clients. */
 	UFUNCTION(BlueprintCallable, Category = "Audio|OSC")
 	TSet<FString> GetWhitelistedClients() const;
diff --git a/Engine/Plugins/Runtime/OSC/Source/OSC/Public/OSCStream.h b/Engine/Plugins/Runtime/OSC/Source/OSC/Public/OSCStream.h
--- a/Engine/Plugins/Runtime/OSC/Source/OSC/Public/OSCStream.h
+++ b/Engine/Plugins/Runtime/OSC/Source/OSC/Public/OSCStream.h
@@ -8,7 +8,6 @@ class FOSCStream
 public:
 	FOSCStream();
 	FOSCStream(const uint8* InData, int32 InSize);
-	FOSCStream(int32 InSize);
 
 	/** Returns the stream buffer data */
 	const uint8* GetData() const;
@@ -92,5 +91,8 @@ private:
 
 	/** Current buffer position. */
 	int32 Position;
+
+	/** Whether stream is used to read (true) or write (false) */
+	bool bIsReadStream;
 };
 
