Fix bad macro syntax

diff --git a/Engine/Plugins/Runtime/OSC/Source/OSC/Private/OSCManager.cpp b/Engine/Plugins/Runtime/OSC/Source/OSC/Private/OSCManager.cpp
--- a/Engine/Plugins/Runtime/OSC/Source/OSC/Private/OSCManager.cpp
+++ b/Engine/Plugins/Runtime/OSC/Source/OSC/Private/OSCManager.cpp
@@ -17,7 +17,7 @@
 #include "UObject/UObjectIterator.h"
 
 
-#define OSC_LOG_INVALID_TYPE_AT_INDEX(Type, Index, Msg) UE_LOG(LogOSC, Warning, TEXT("OSC Message Parse Failed: OSCType not %s: index '%i', OSCAddress '%s'"), TEXT(##Type), Index, *Msg.GetAddress().GetFullPath())
+#define OSC_LOG_INVALID_TYPE_AT_INDEX(TypeStr, Index, Msg) UE_LOG(LogOSC, Warning, TEXT("OSC Message Parse Failed: OSCType not %s: index '%i', OSCAddress '%s'"), TypeStr, Index, *Msg.GetAddress().GetFullPath())
 
 namespace OSC
 {
@@ -355,7 +355,7 @@ bool UOSCManager::GetAddress(const FOSCMessage& InMessage, const int32 InIndex,
 			OutValue = FOSCAddress(OSCType->GetString());
 			return OutValue.IsValidPath();
 		}
-		OSC_LOG_INVALID_TYPE_AT_INDEX("String (OSCAddress)", InIndex, InMessage);
+		OSC_LOG_INVALID_TYPE_AT_INDEX(TEXT("String (OSCAddress)"), InIndex, InMessage);
 	}
 
 	OutValue = FOSCAddress();
@@ -393,7 +393,7 @@ bool UOSCManager::GetFloat(const FOSCMessage& InMessage, const int32 InIndex, fl
 			OutValue = OSCType->GetFloat();
 			return true;
 		}
-		OSC_LOG_INVALID_TYPE_AT_INDEX("Float", InIndex, InMessage);
+		OSC_LOG_INVALID_TYPE_AT_INDEX(TEXT("Float"), InIndex, InMessage);
 	}
 
 	return false;
@@ -426,7 +426,7 @@ bool UOSCManager::GetInt32(const FOSCMessage& InMessage, const int32 InIndex, in
 			OutValue = OSCType->GetInt32();
 			return true;
 		}
-		OSC_LOG_INVALID_TYPE_AT_INDEX("Int32", InIndex, InMessage);
+		OSC_LOG_INVALID_TYPE_AT_INDEX(TEXT("Int32"), InIndex, InMessage);
 	}
 
 	return false;
@@ -459,7 +459,7 @@ bool UOSCManager::GetInt64(const FOSCMessage& InMessage, const int32 InIndex, in
 			OutValue = OSCType->GetInt64();
 			return true;
 		}
-		OSC_LOG_INVALID_TYPE_AT_INDEX("Int64", InIndex, InMessage);
+		OSC_LOG_INVALID_TYPE_AT_INDEX(TEXT("Int64"), InIndex, InMessage);
 	}
 
 	return false;
@@ -491,7 +491,7 @@ bool UOSCManager::GetString(const FOSCMessage& InMessage, const int32 InIndex, F
 			OutValue = OSCType->GetString();
 			return true;
 		}
-		OSC_LOG_INVALID_TYPE_AT_INDEX("String", InIndex, InMessage);
+		OSC_LOG_INVALID_TYPE_AT_INDEX(TEXT("String"), InIndex, InMessage);
 	}
 
 	OutValue.Reset();
@@ -525,7 +525,7 @@ bool UOSCManager::GetBool(const FOSCMessage& InMessage, const int32 InIndex, boo
 			OutValue = OSCType->GetBool();
 			return true;
 		}
-		OSC_LOG_INVALID_TYPE_AT_INDEX("Bool", InIndex, InMessage);
+		OSC_LOG_INVALID_TYPE_AT_INDEX(TEXT("Bool"), InIndex, InMessage);
 	}
 
 	return false;
@@ -558,7 +558,7 @@ bool UOSCManager::GetBlob(const FOSCMessage& InMessage, const int32 InIndex, TAr
 			OutValue = OSCType->GetBlob();
 			return true;
 		}
-		OSC_LOG_INVALID_TYPE_AT_INDEX("Blob", InIndex, InMessage);
+		OSC_LOG_INVALID_TYPE_AT_INDEX(TEXT("Blob"), InIndex, InMessage);
 	}
 
 	return false;
