- Add new cvars to manipulate and test OSC client/server connections as well as query state of initialized respective uobjects
- Ensure OSC clients/servers are stopped on destruction
- Fix-ups for BP API (refs on wrong side of node)

diff --git a/Engine/Plugins/Runtime/OSC/Source/OSC/Private/OSCClientProxy.cpp b/Engine/Plugins/Runtime/OSC/Source/OSC/Private/OSCClientProxy.cpp
--- a/Engine/Plugins/Runtime/OSC/Source/OSC/Private/OSCClientProxy.cpp
+++ b/Engine/Plugins/Runtime/OSC/Source/OSC/Private/OSCClientProxy.cpp
@@ -22,6 +22,11 @@ FOSCClientProxy::FOSCClientProxy(const FString& InClientName)
 	ClientName = InClientName;
 }
 
+FOSCClientProxy::~FOSCClientProxy()
+{
+	Stop();
+}
+
 void FOSCClientProxy::GetSendIPAddress(FString& InIPAddress, int32& Port) const
 {
 	const bool bAppendPort = false;
diff --git a/Engine/Plugins/Runtime/OSC/Source/OSC/Private/OSCClientProxy.h b/Engine/Plugins/Runtime/OSC/Source/OSC/Private/OSCClientProxy.h
--- a/Engine/Plugins/Runtime/OSC/Source/OSC/Private/OSCClientProxy.h
+++ b/Engine/Plugins/Runtime/OSC/Source/OSC/Private/OSCClientProxy.h
@@ -14,6 +14,7 @@ class OSC_API FOSCClientProxy : public IOSCClientProxy
 {
 public:
 	FOSCClientProxy(const FString& InClientName);
+	virtual ~FOSCClientProxy();
 
 	void GetSendIPAddress(FString& InIPAddress, int32& Port) const override;
 	bool SetSendIPAddress(const FString& InIPAddress, const int32 Port) override;
diff --git a/Engine/Plugins/Runtime/OSC/Source/OSC/Private/OSCManager.cpp b/Engine/Plugins/Runtime/OSC/Source/OSC/Private/OSCManager.cpp
--- a/Engine/Plugins/Runtime/OSC/Source/OSC/Private/OSCManager.cpp
+++ b/Engine/Plugins/Runtime/OSC/Source/OSC/Private/OSCManager.cpp
@@ -21,6 +21,9 @@
 
 namespace OSC
 {
+	static const int32 DefaultClientPort = 8094;
+	static const int32 DefaultServerPort = 8095;
+
 	// Returns true if provided address was null and was able to
 	// override with local host address, false if not.
 	bool GetLocalHostAddress(FString& InAddress)
@@ -65,9 +68,9 @@ namespace OSC
 } // namespace OSC
 
 
-static FAutoConsoleCommand GOSCGetServerDiag(
+static FAutoConsoleCommand GOSCPrintServers(
 	TEXT("osc.servers"),
-	TEXT("Prints diagnostic information pertaining to the current OSC servers to the output log."),
+	TEXT("Prints diagnostic information pertaining to the currently initialized OSC servers objects to the output log."),
 	FConsoleCommandDelegate::CreateStatic(
 		[]()
 		{
@@ -81,13 +84,13 @@ static FAutoConsoleCommand GOSCGetServerDiag(
 				if (UOSCServer* Server = *Itr)
 				{
 					FString ToPrint = TEXT("    ") + Server->GetName();
+					ToPrint.Appendf(TEXT(" (Id: %u"), Server->GetUniqueID());
 					if (UWorld* World = Server->GetWorld())
 					{
-						ToPrint += TEXT("(") + World->GetName() + TEXT(")");
+						ToPrint.Appendf(TEXT(", World: %s"), *World->GetName());
 					}
 
-					ToPrint += TEXT(", ");
-					ToPrint += Server->GetIpAddress(true /* bIncludePort */);
+					ToPrint.Appendf(TEXT(", IP: %s)"), *Server->GetIpAddress(true /* bIncludePort */));
 					ToPrint += Server->IsActive() ? TEXT(" [Active]") : TEXT(" [Inactive]");
 
 					UE_LOG(LogOSC, Display, TEXT("%s"), *ToPrint);
@@ -108,9 +111,110 @@ static FAutoConsoleCommand GOSCGetServerDiag(
 	)
 );
 
-static FAutoConsoleCommand GOSCGetClientDiag(
+static FAutoConsoleCommand GOSCServerConnect(
+	TEXT("osc.server.connect"),
+	TEXT("Connects or reconnects the osc mix server with the provided name\n"
+		"(see \"osc.servers\" for a list of available servers and their respective names). Args:\n"
+		"Name - Object name of server to (re)connect"
+		"Address - IP Address to connect to (default: LocalHost)"
+		"Port - Port to connect to (default: 8095)"),
+	FConsoleCommandWithArgsDelegate::CreateStatic(
+		[](const TArray<FString>& Args)
+		{
+			FString SrvName;
+			if (Args.Num() > 0)
+			{
+				SrvName = Args[0];
+			}
+
+			FString IPAddr;
+			OSC::GetLocalHostAddress(IPAddr);
+			if (Args.Num() > 1)
+			{
+				IPAddr = Args[1];
+			}
+
+			int32 Port = OSC::DefaultServerPort;
+			if (Args.Num() > 2)
+			{
+				Port = FCString::Atoi(*Args[2]);
+			}
+
+			for (TObjectIterator<UOSCServer> Itr; Itr; ++Itr)
+			{
+				if (UOSCServer* Server = *Itr)
+				{
+					if (Server->GetName() == SrvName)
+					{
+						Server->Stop();
+						if (Server->SetAddress(IPAddr, Port))
+						{
+							Server->Listen();
+						}
+						return;
+					}
+				}
+			}
+
+			UE_LOG(LogOSC, Warning, TEXT("Server object with name '%s' not found, (re)connect not performed."), *SrvName);
+		}
+	)
+);
+
+static FAutoConsoleCommand GOSCServerConnectById(
+	TEXT("osc.server.connectById"),
+	TEXT("Connects or reconnects the osc mix server with the provided object id\n"
+		"(see \"osc.servers\" for a list of available servers and their respective ids). Args:\n"
+		"Id - Object Id of client to (re)connect"
+		"Address - IP Address to (re)connect to (default: LocalHost)"
+		"Port - Port to (re)connect to (default: 8095)"),
+	FConsoleCommandWithArgsDelegate::CreateStatic(
+		[](const TArray<FString>& Args)
+		{
+			if (Args.Num() == 0)
+			{
+				return;
+			}
+
+			const int32 SrvId = FCString::Atoi(*Args[0]);
+
+			FString IPAddr;
+			OSC::GetLocalHostAddress(IPAddr);
+			if (Args.Num() > 1)
+			{
+				IPAddr = Args[1];
+			}
+
+			int32 Port = OSC::DefaultServerPort;
+			if (Args.Num() > 2)
+			{
+				Port = FCString::Atoi(*Args[2]);
+			}
+
+			for (TObjectIterator<UOSCServer> Itr; Itr; ++Itr)
+			{
+				if (UOSCServer* Server = *Itr)
+				{
+					if (Server->GetUniqueID() == SrvId)
+					{
+						Server->Stop();
+						if (Server->SetAddress(IPAddr, Port))
+						{
+							Server->Listen();
+						}
+						return;
+					}
+				}
+			}
+
+			UE_LOG(LogOSC, Warning, TEXT("Server object with id '%u' not found, (re)connect not performed."), SrvId);
+		}
+	)
+);
+
+static FAutoConsoleCommand GOSCPrintClients(
 	TEXT("osc.clients"),
-	TEXT("Prints diagnostic information pertaining to the current OSC clients to the output log."),
+	TEXT("Prints diagnostic information pertaining to the currently initialized OSC client objects to the output log."),
 	FConsoleCommandDelegate::CreateStatic(
 		[]()
 		{
@@ -123,12 +227,11 @@ static FAutoConsoleCommand GOSCGetClientDiag(
 			{
 				if (UOSCClient* Client = *Itr)
 				{
-					FString ToPrint;
-					ToPrint += TEXT("\n    ") + Client->GetName();
-
+					FString ToPrint = TEXT("    ") + Client->GetName();
+					ToPrint.Appendf(TEXT(" (Id: %u"), Client->GetUniqueID());
 					if (UWorld* World = Client->GetWorld())
 					{
-						ToPrint += TEXT("(") + World->GetName() + TEXT(")");
+						ToPrint.Appendf(TEXT(", World: %s"), *World->GetName());
 					}
 
 					FString IPAddrStr;
@@ -137,6 +240,7 @@ static FAutoConsoleCommand GOSCGetClientDiag(
 					ToPrint += TEXT(", ") + IPAddrStr + TEXT(":");
 					ToPrint.AppendInt(Port);
 					ToPrint += Client->IsActive() ? TEXT(" [Active]") : TEXT(" [Inactive]");
+
 					UE_LOG(LogOSC, Display, TEXT("%s"), *ToPrint);
 				}
 			}
@@ -144,47 +248,158 @@ static FAutoConsoleCommand GOSCGetClientDiag(
 	)
 );
 
+static FAutoConsoleCommand GOSCClientConnect(
+	TEXT("osc.client.connect"),
+	TEXT("Connects (or reconnects) the osc mix client with the provided name\n"
+		"(see \"osc.clients\" for a list of available clients and their respective ids). Args:\n"
+		"Name - Object name of client to (re)connect"
+		"Address - IP Address to (re)connect to (default: LocalHost)"
+		"Port - Port to (re)connect to (default: 8094)"),
+	FConsoleCommandWithArgsDelegate::CreateStatic(
+		[](const TArray< FString >& Args)
+		{
+			if (Args.Num() == 0)
+			{
+				return;
+			}
+			const FString CliName = Args[0];
 
-UOSCServer* UOSCManager::CreateOSCServer(FString InReceiveIPAddress, int32 InPort, bool bInMulticastLoopback, bool bInStartListening)
+			FString IPAddr;
+			OSC::GetLocalHostAddress(IPAddr);
+			if (Args.Num() > 1)
+			{
+				IPAddr = Args[1];
+			}
+
+			int32 Port = OSC::DefaultClientPort;
+			if (Args.Num() > 2)
+			{
+				Port = FCString::Atoi(*Args[2]);
+			}
+
+			for (TObjectIterator<UOSCClient> Itr; Itr; ++Itr)
+			{
+				if (UOSCClient* Client = *Itr)
+				{
+					if (Client->GetName() == CliName)
+					{
+						Client->Connect();
+						Client->SetSendIPAddress(IPAddr, Port);
+						return;
+					}
+				}
+			}
+
+			UE_LOG(LogOSC, Warning, TEXT("Client object with name '%s' not found, (re)connect not performed."), *CliName);
+		}
+	)
+);
+
+static FAutoConsoleCommand GOSCClientConnectById(
+	TEXT("osc.client.connectById"),
+	TEXT("Connects (or reconnects) the osc mix client with the provided object id\n"
+		"(see \"osc.clients\" for a list of available clients and their respective ids). Args:\n"
+		"Id - Object Id of client to (re)connect"
+		"Address - IP Address to (re)connect to (default: LocalHost)"
+		"Port - Port to (re)connect to (default: 8094)"),
+	FConsoleCommandWithArgsDelegate::CreateStatic(
+		[](const TArray< FString >& Args)
+		{
+			int32 CliId = INDEX_NONE;
+			if (Args.Num() > 0)
+			{
+				CliId = FCString::Atoi(*Args[0]);
+			}
+
+			FString IPAddr;
+			OSC::GetLocalHostAddress(IPAddr);
+			if (Args.Num() > 1)
+			{
+				IPAddr = Args[1];
+			}
+
+			int32 Port = OSC::DefaultClientPort;
+			if (Args.Num() > 2)
+			{
+				Port = FCString::Atoi(*Args[2]);
+			}
+
+			for (TObjectIterator<UOSCClient> Itr; Itr; ++Itr)
+			{
+				if (UOSCClient* Client = *Itr)
+				{
+					if (Client->GetUniqueID() == CliId)
+					{
+						Client->Connect();
+						Client->SetSendIPAddress(IPAddr, Port);
+						return;
+					}
+				}
+			}
+
+			UE_LOG(LogOSC, Warning, TEXT("Client object with id '%u' not found, (re)connect not performed."), CliId);
+		}
+	)
+);
+
+UOSCServer* UOSCManager::CreateOSCServer(FString InReceiveIPAddress, int32 InPort, bool bInMulticastLoopback, bool bInStartListening, FString ServerName)
 {
 	if (OSC::GetLocalHostAddress(InReceiveIPAddress))
 	{
 		UE_LOG(LogOSC, Display, TEXT("OSCServer ReceiveAddress not specified. Using LocalHost IP: '%s'"), *InReceiveIPAddress);
 	}
 
-	UOSCServer* NewOSCServer = NewObject<UOSCServer>();
-	NewOSCServer->Connect();
-	NewOSCServer->SetMulticastLoopback(bInMulticastLoopback);
-	if (NewOSCServer->SetAddress(InReceiveIPAddress, InPort))
+	if (ServerName.IsEmpty())
 	{
-		if (bInStartListening)
-		{
-			NewOSCServer->Listen();
-		}
+		ServerName = FString::Printf(TEXT("OSCServer_%s"), *FGuid::NewGuid().ToString(EGuidFormats::Short));
 	}
-	else
+
+	if (UOSCServer* NewOSCServer = NewObject<UOSCServer>(GetTransientPackage(), FName(*ServerName)))
 	{
-		UE_LOG(LogOSC, Warning, TEXT("Failed to parse ReceiveAddress '%s' for OSCServer."), *InReceiveIPAddress);
+		NewOSCServer->Connect();
+		NewOSCServer->SetMulticastLoopback(bInMulticastLoopback);
+		if (NewOSCServer->SetAddress(InReceiveIPAddress, InPort))
+		{
+			if (bInStartListening)
+			{
+				NewOSCServer->Listen();
+			}
+		}
+		else
+		{
+			UE_LOG(LogOSC, Warning, TEXT("Failed to parse ReceiveAddress '%s' for OSCServer."), *InReceiveIPAddress);
+		}
+
+		return NewOSCServer;
 	}
 
-	return NewOSCServer;
+	return nullptr;
 }
 
-UOSCClient* UOSCManager::CreateOSCClient(FString InSendIPAddress, int32 InPort)
+UOSCClient* UOSCManager::CreateOSCClient(FString InSendIPAddress, int32 InPort, FString ClientName)
 {
 	if (OSC::GetLocalHostAddress(InSendIPAddress))
 	{
 		UE_LOG(LogOSC, Display, TEXT("OSCClient SendAddress not specified. Using LocalHost IP: '%s'"), *InSendIPAddress);
 	}
 
-	UOSCClient* NewOSCClient = NewObject<UOSCClient>();
-	NewOSCClient->Connect();
-	if (!NewOSCClient->SetSendIPAddress(InSendIPAddress, InPort))
+	if (ClientName.IsEmpty())
 	{
-		UE_LOG(LogOSC, Warning, TEXT("Failed to parse SendAddress '%s' for OSCClient. Client unable to send new messages."), *InSendIPAddress);
+		ClientName = FString::Printf(TEXT("OSCClient_%s"), *FGuid::NewGuid().ToString(EGuidFormats::Short));
 	}
 
-	return NewOSCClient;
+	if (UOSCClient* NewOSCClient = NewObject<UOSCClient>(GetTransientPackage(), FName(*ClientName)))
+	{
+		NewOSCClient->Connect();
+		if (!NewOSCClient->SetSendIPAddress(InSendIPAddress, InPort))
+		{
+			UE_LOG(LogOSC, Warning, TEXT("Failed to parse SendAddress '%s' for OSCClient. Client unable to send new messages."), *InSendIPAddress);
+		}
+
+		return NewOSCClient;
+	}
+
+	return nullptr;
 }
 
 FOSCMessage& UOSCManager::ClearMessage(FOSCMessage& OutMessage)
@@ -270,7 +485,7 @@ FOSCMessage& UOSCManager::AddString(FOSCMessage& OutMessage, FString InValue)
 	return OutMessage;
 }
 
-FOSCMessage& UOSCManager::AddBlob(FOSCMessage& OutMessage, TArray<uint8>& OutValue)
+FOSCMessage& UOSCManager::AddBlob(FOSCMessage& OutMessage, const TArray<uint8>& OutValue)
 {
 	const TSharedPtr<FOSCMessagePacket>& MessagePacket = StaticCastSharedPtr<FOSCMessagePacket>(OutMessage.GetPacket());
 	MessagePacket->GetArguments().Add(FOSCType(OutValue));
diff --git a/Engine/Plugins/Runtime/OSC/Source/OSC/Private/OSCServerProxy.cpp b/Engine/Plugins/Runtime/OSC/Source/OSC/Private/OSCServerProxy.cpp
--- a/Engine/Plugins/Runtime/OSC/Source/OSC/Private/OSCServerProxy.cpp
+++ b/Engine/Plugins/Runtime/OSC/Source/OSC/Private/OSCServerProxy.cpp
@@ -20,6 +20,11 @@ FOSCServerProxy::FOSCServerProxy(UOSCServer& InServer)
 {
 }
 
+FOSCServerProxy::~FOSCServerProxy()
+{
+	Stop();
+}
+
 void FOSCServerProxy::OnPacketReceived(const FArrayReaderPtr& Data, const FIPv4Endpoint& Endpoint)
 {
 	TSharedPtr<IOSCPacket> Packet = IOSCPacket::CreatePacket(Data->GetData());
diff --git a/Engine/Plugins/Runtime/OSC/Source/OSC/Private/OSCServerProxy.h b/Engine/Plugins/Runtime/OSC/Source/OSC/Private/OSCServerProxy.h
--- a/Engine/Plugins/Runtime/OSC/Source/OSC/Private/OSCServerProxy.h
+++ b/Engine/Plugins/Runtime/OSC/Source/OSC/Private/OSCServerProxy.h
@@ -15,7 +15,7 @@ class OSC_API FOSCServerProxy : public IOSCServerProxy
 {
 public:
 	FOSCServerProxy(UOSCServer& InServer);
-	virtual ~FOSCServerProxy() { }
+	virtual ~FOSCServerProxy();
 
 	bool GetMulticastLoopback() const override;
 	FString GetIpAddress() const override;
diff --git a/Engine/Plugins/Runtime/OSC/Source/OSC/Public/OSCManager.h b/Engine/Plugins/Runtime/OSC/Source/OSC/Public/OSCManager.h
--- a/Engine/Plugins/Runtime/OSC/Source/OSC/Public/OSCManager.h
+++ b/Engine/Plugins/Runtime/OSC/Source/OSC/Public/OSCManager.h
@@ -22,16 +22,16 @@ class OSC_API UOSCManager : public UBlueprintFunctionLibrary
 	GENERATED_BODY()
 
 public:
-	// Creates an OSC Server.  If left empty (or '0'),
+	// Creates an OSC Server.  If ReceiveIPAddress left empty (or '0'),
 	// attempts to use LocalHost IP address. If StartListening set,
 	// immediately begins listening on creation.
 	UFUNCTION(BlueprintCallable, Category = "Audio|OSC")
-	static UOSCServer* CreateOSCServer(FString ReceiveIPAddress, int32 Port, bool bMulticastLoopback, bool bStartListening);
+	static UOSCServer* CreateOSCServer(FString ReceiveIPAddress, int32 Port, bool bMulticastLoopback, bool bStartListening, FString ServerName);
 
-	// Creates an OSC Client.  If left empty (or '0'), attempts to use
+	// Creates an OSC Client.  If SendIPAddress left empty (or '0'), attempts to use
 	// attempts to use LocalHost IP address.
 	UFUNCTION(BlueprintCallable, Category = "Audio|OSC")
-	static UOSCClient* CreateOSCClient(FString SendIPAddress, int32 Port);
+	static UOSCClient* CreateOSCClient(FString SendIPAddress, int32 Port, FString ClientName);
 
 	/** Adds provided message packet to bundle. */
 	UFUNCTION(BlueprintCallable, Category = "Audio|OSC", meta = (DisplayName = "Add OSC Message to Bundle", Keywords = "osc message bundle"))
@@ -46,7 +46,7 @@ public:
 	static UPARAM(DisplayName = "Bundles") TArray<FOSCBundle> GetBundlesFromBundle(const FOSCBundle& Bundle);
 
 	/** Returns message found in bundle at ordered index. */
-	UFUNCTION(BlueprintCallable, Category = "Audio|OSC", meta = (DisplayName = "Get OSC Messages From Bundle", Keywords = "osc bundle message"))
+	UFUNCTION(BlueprintCallable, Category = "Audio|OSC", meta = (DisplayName = "Get OSC Message From Bundle At Index", Keywords = "osc bundle message"))
 	static UPARAM(DisplayName = "Message") FOSCMessage GetMessageFromBundle(const FOSCBundle& Bundle, int32 Index, bool& bSucceeded);
 
 	/** Fills array with messages found in bundle. */
@@ -79,11 +79,11 @@ public:
 
 	/** Adds string value to end of OSCMessage */
 	UFUNCTION(BlueprintCallable, Category = "Audio|OSC", meta = (DisplayName = "Add String to OSC Message", Keywords = "osc message"))
-	static UPARAM(DisplayName = "Message") FOSCMessage& AddString(UPARAM(ref) FOSCMessage& Message, FString Value);
+	static UPARAM(DisplayName = "Message") FOSCMessage& AddString(UPARAM(ref) FOSCMessage& Message, UPARAM(ref) FString Value);
 
 	/** Adds blob value to end of OSCMessage */
 	UFUNCTION(BlueprintCallable, Category = "Audio|OSC", meta = (DisplayName = "Add Blob to OSC Message", Keywords = "osc message"))
-	static UPARAM(DisplayName = "Message") FOSCMessage& AddBlob(UPARAM(ref) FOSCMessage& Message, TArray<uint8>& Value);
+	static UPARAM(DisplayName = "Message") FOSCMessage& AddBlob(UPARAM(ref) FOSCMessage& Message, const TArray<uint8>& Value);
 
 	/** Adds boolean value to end of OSCMessage */
 	UFUNCTION(BlueprintCallable, Category = "Audio|OSC", meta = (DisplayName = "Add Bool to OSC Message", Keywords = "osc message"))
