Add native delegates to OSC server
Pass address and port to callbacks

diff --git a/Engine/Plugins/Runtime/OSC/Source/OSC/Private/OSCServer.cpp b/Engine/Plugins/Runtime/OSC/Source/OSC/Private/OSCServer.cpp
--- a/Engine/Plugins/Runtime/OSC/Source/OSC/Private/OSCServer.cpp
+++ b/Engine/Plugins/Runtime/OSC/Source/OSC/Private/OSCServer.cpp
@@ -177,9 +177,10 @@ void UOSCServer::EnqueuePacket(TSharedPtr<IOSCPacket> Packet)
 	OSCPackets.Enqueue(Packet);
 }
 
-void UOSCServer::DispatchBundle(const FString& InIPAddress, const FOSCBundle& InBundle)
+void UOSCServer::DispatchBundle(const FString& InIPAddress, uint16 Port, const FOSCBundle& InBundle)
 {
-	OnOscBundleReceived.Broadcast(InBundle);
+	OnOscBundleReceived.Broadcast(InBundle, InIPAddress, Port);
+	OnOscBundleReceivedNative.Broadcast(InBundle, InIPAddress, Port);
 
 	TSharedPtr<FOSCBundlePacket> BundlePacket = StaticCastSharedPtr<FOSCBundlePacket>(InBundle.GetPacket());
 	FOSCBundlePacket::FPacketBundle Packets = BundlePacket->GetPackets();
@@ -187,11 +188,11 @@ void UOSCServer::DispatchBundle(const FString& InIPAddress, const FOSCBundle& In
 	{
 		if (Packet->IsMessage())
 		{
-			DispatchMessage(InIPAddress, FOSCMessage(Packet));
+			DispatchMessage(InIPAddress, Port, FOSCMessage(Packet));
 		}
 		else if (Packet->IsBundle())
 		{
-			DispatchBundle(InIPAddress, FOSCBundle(Packet));
+			DispatchBundle(InIPAddress, Port, FOSCBundle(Packet));
 		}
 		else
 		{
@@ -200,9 +201,11 @@ void UOSCServer::DispatchBundle(const FString& InIPAddress, const FOSCBundle& In
 	}
 }
 
-void UOSCServer::DispatchMessage(const FString& InIPAddress, const FOSCMessage& InMessage)
+void UOSCServer::DispatchMessage(const FString& InIPAddress, uint16 Port, const FOSCMessage& InMessage)
 {
-	OnOscMessageReceived.Broadcast(InMessage);
+	OnOscMessageReceived.Broadcast(InMessage, InIPAddress, Port);
+	OnOscMessageReceivedNative.Broadcast(InMessage, InIPAddress, Port);
+
 	UE_LOG(LogOSC, Verbose, TEXT("Message received from endpoint '%s', OSCAddress of '%s'."), *InIPAddress, *InMessage.GetAddress().GetFullPath());
 
 	for (const TPair<FOSCAddress, FOSCDispatchMessageEvent>& Pair : AddressPatterns)
@@ -210,7 +213,7 @@ void UOSCServer::DispatchMessage(const FString& InIPAddress, const FOSCMessage&
 		const FOSCDispatchMessageEvent& DispatchEvent = Pair.Value;
 		if (Pair.Key.Matches(InMessage.GetAddress()))
 		{
-			DispatchEvent.Broadcast(Pair.Key, InMessage);
+			DispatchEvent.Broadcast(Pair.Key, InMessage, InIPAddress, Port);
 			UE_LOG(LogOSC, Verbose, TEXT("Message dispatched from endpoint '%s', OSCAddress path of '%s' matched OSCAddress pattern '%s'."),
 				*InIPAddress,
 				*InMessage.GetAddress().GetFullPath(),
@@ -219,18 +222,18 @@ void UOSCServer::DispatchMessage(const FString& InIPAddress, const FOSCMessage&
 	}
 }
 
-void UOSCServer::OnPacketReceived(const FString& InIPAddress)
+void UOSCServer::OnPacketReceived(const FString& InIPAddress, uint16 Port)
 {
 	TSharedPtr<IOSCPacket> Packet;
 	while (OSCPackets.Dequeue(Packet))
 	{
 		if (Packet->IsMessage())
 		{
-			DispatchMessage(InIPAddress, FOSCMessage(Packet));
+			DispatchMessage(InIPAddress, Port, FOSCMessage(Packet));
 		}
 		else if (Packet->IsBundle())
 		{
-			DispatchBundle(InIPAddress, FOSCBundle(Packet));
+			DispatchBundle(InIPAddress, Port, FOSCBundle(Packet));
 		}
 		else
 		{
diff --git a/Engine/Plugins/Runtime/OSC/Source/OSC/Private/OSCServerProxy.cpp b/Engine/Plugins/Runtime/OSC/Source/OSC/Private/OSCServerProxy.cpp
--- a/Engine/Plugins/Runtime/OSC/Source/OSC/Private/OSCServerProxy.cpp
+++ b/Engine/Plugins/Runtime/OSC/Source/OSC/Private/OSCServerProxy.cpp
@@ -48,7 +48,7 @@ void FOSCServerProxy::OnPacketReceived(const FArrayReaderPtr& Data, const FIPv4E
 			return;
 		}
 
-		Server->OnPacketReceived(Endpoint.Address.ToString());
+		Server->OnPacketReceived(Endpoint.Address.ToString(), Endpoint.Port);
 	}, GET_STATID(STAT_OSCServerOnPacketReceived), nullptr, ENamedThreads::GameThread);
 }
 
diff --git a/Engine/Plugins/Runtime/OSC/Source/OSC/Public/OSCServer.h b/Engine/Plugins/Runtime/OSC/Source/OSC/Public/OSCServer.h
--- a/Engine/Plugins/Runtime/OSC/Source/OSC/Public/OSCServer.h
+++ b/Engine/Plugins/Runtime/OSC/Source/OSC/Public/OSCServer.h
@@ -19,10 +19,12 @@ class FSocket;
 
 
 // Delegates
-DECLARE_DYNAMIC_MULTICAST_DELEGATE_OneParam(FOSCReceivedMessageEvent, const FOSCMessage&, Message);
-DECLARE_DYNAMIC_MULTICAST_DELEGATE_TwoParams(FOSCDispatchMessageEvent, const FOSCAddress&, AddressPattern, const FOSCMessage&, Message);
-DECLARE_DYNAMIC_MULTICAST_DELEGATE_OneParam(FOSCReceivedBundleEvent, const FOSCBundle&, Bundle);
-DECLARE_DYNAMIC_DELEGATE_TwoParams(FOSCDispatchMessageEventBP, const FOSCAddress&, AddressPattern, const FOSCMessage&, Message);
+DECLARE_DYNAMIC_MULTICAST_DELEGATE_ThreeParams(FOSCReceivedMessageEvent, const FOSCMessage&, Message, const FString&, IPAddress, int32, Port);
+DECLARE_MULTICAST_DELEGATE_ThreeParams(FOSCReceivedMessageNativeEvent, const FOSCMessage&, const FString& /*IPAddress*/, uint16 /*Port*/);
+DECLARE_DYNAMIC_MULTICAST_DELEGATE_FourParams(FOSCDispatchMessageEvent, const FOSCAddress&, AddressPattern, const FOSCMessage&, Message, const FString&, IPAddress, int32, Port);
+DECLARE_DYNAMIC_MULTICAST_DELEGATE_ThreeParams(FOSCReceivedBundleEvent, const FOSCBundle&, Bundle, const FString&, IPAddress, int32, Port);
+DECLARE_MULTICAST_DELEGATE_ThreeParams(FOSCReceivedBundleNativeEvent, const FOSCBundle&, const FString& /*IPAddress*/, uint16 /*Port*/);
+DECLARE_DYNAMIC_DELEGATE_FourParams(FOSCDispatchMessageEventBP, const FOSCAddress&, AddressPattern, const FOSCMessage&, Message, const FString&, IPAddress, int32, Port);
 
 DECLARE_STATS_GROUP(TEXT("OSC Commands"), STATGROUP_OSCNetworkCommands, STATCAT_Advanced);
 
@@ -84,10 +86,16 @@ public:
 	UPROPERTY(BlueprintAssignable, Category = "Audio|OSC")
 	FOSCReceivedMessageEvent OnOscMessageReceived;
 
+	/** Native event that gets called when an OSC message is received. */
+	FOSCReceivedMessageNativeEvent OnOscMessageReceivedNative;
+
 	/** Event that gets called when an OSC bundle is received. */
 	UPROPERTY(BlueprintAssignable, Category = "Audio|OSC")
 	FOSCReceivedBundleEvent OnOscBundleReceived;
 
+	/** Native event that gets called when an OSC bundle is received. */
+	FOSCReceivedBundleNativeEvent OnOscBundleReceivedNative;
+
 	/** When set to true, server will only process received 
 	  * messages from whitelisted clients.
 	  */
@@ -145,17 +153,17 @@ public:
 	void EnqueuePacket(TSharedPtr<IOSCPacket>);
 
 	/** Callback for when packet is received by server */
-	void OnPacketReceived(const FString& InIPAddress);
+	void OnPacketReceived(const FString& InIPAddress, uint16 Port);
 
 protected:
 	void BeginDestroy() override;
 
 private:
 	/** Dispatches provided bundle received */
-	void DispatchBundle(const FString& InIPAddress, const FOSCBundle& InBundle);
+	void DispatchBundle(const FString& InIPAddress, uint16 Port, const FOSCBundle& InBundle);
 
 	/** Dispatches provided message received */
-	void DispatchMessage(const FString& InIPAddress, const FOSCMessage& InMessage);
+	void DispatchMessage(const FString& InIPAddress, uint16 Port, const FOSCMessage& InMessage);
 
 	/** Pointer to internal implementation of server proxy */
 	TUniquePtr<IOSCServerProxy> ServerProxy;
