Change names of OSC Types to avoid collisions with common macro names & reduce type tag size to uint8

diff --git a/Engine/Plugins/Runtime/OSC/Source/OSC/Private/OSCMessagePacket.cpp b/Engine/Plugins/Runtime/OSC/Source/OSC/Private/OSCMessagePacket.cpp
--- a/Engine/Plugins/Runtime/OSC/Source/OSC/Private/OSCMessagePacket.cpp
+++ b/Engine/Plugins/Runtime/OSC/Source/OSC/Private/OSCMessagePacket.cpp
@@ -62,44 +62,44 @@ void FOSCMessagePacket::WriteData(FOSCStream& Stream)
 	{
 		switch (OSCType.GetTypeTag())
 		{
-		case EOSCTypeTag::CHAR:
+		case EOSCTypeTag::OSC_CHAR:
 			Stream.WriteChar(OSCType.GetChar());
 			break;
-		case EOSCTypeTag::INT32:
+		case EOSCTypeTag::OSC_INT32:
 			Stream.WriteInt32(OSCType.GetInt32());
 			break;
-		case EOSCTypeTag::FLOAT:
+		case EOSCTypeTag::OSC_FLOAT:
 			Stream.WriteFloat(OSCType.GetFloat());
 			break;
-		case EOSCTypeTag::DOUBLE:
+		case EOSCTypeTag::OSC_DOUBLE:
 			Stream.WriteDouble(OSCType.GetDouble());
 			break;
-		case EOSCTypeTag::INT64:
+		case EOSCTypeTag::OSC_INT64:
 			Stream.WriteInt64(OSCType.GetInt64());
 			break;
-		case EOSCTypeTag::TIME:
+		case EOSCTypeTag::OSC_TIME:
 			Stream.WriteUInt64(OSCType.GetTimeTag());
 			break;
-		case EOSCTypeTag::STRING:
+		case EOSCTypeTag::OSC_STRING:
 			Stream.WriteString(OSCType.GetString());
 			break;
-		case EOSCTypeTag::BLOB:
+		case EOSCTypeTag::OSC_BLOB:
 		{
 			TArray<uint8> blob = OSCType.GetBlob();
 			Stream.WriteBlob(blob);
 		}
 		break;
-		case EOSCTypeTag::COLOR:
+		case EOSCTypeTag::OSC_COLOR:
 #if PLATFORM_LITTLE_ENDIAN
 			Stream.WriteInt32(OSCType.GetColor().ToPackedABGR());
 #else
 			Stream.WriteInt32(OSCType.GetColor().ToPackedRGBA());
 #endif
 			break;
-		case EOSCTypeTag::TRUE:
-		case EOSCTypeTag::FALSE:
-		case EOSCTypeTag::NIL:
-		case EOSCTypeTag::INFINITUM:
+		case EOSCTypeTag::OSC_TRUE:
+		case EOSCTypeTag::OSC_FALSE:
+		case EOSCTypeTag::OSC_NIL:
+		case EOSCTypeTag::OSC_INFINITUM:
 			// No values are written for these types
 			break;
 		default:
@@ -131,46 +131,46 @@ void FOSCMessagePacket::ReadData(FOSCStream& Stream)
 		const EOSCTypeTag Tag = static_cast<EOSCTypeTag>(TagTypes[i]);
 		switch (Tag)
 		{
-		case EOSCTypeTag::CHAR:
+		case EOSCTypeTag::OSC_CHAR:
 			Arguments.Add(FOSCType(Stream.ReadChar()));
 			break;
-		case EOSCTypeTag::INT32:
+		case EOSCTypeTag::OSC_INT32:
 			Arguments.Add(FOSCType(Stream.ReadInt32()));
 			break;
-		case EOSCTypeTag::FLOAT:
+		case EOSCTypeTag::OSC_FLOAT:
 			Arguments.Add(FOSCType(Stream.ReadFloat()));
 			break;
-		case EOSCTypeTag::DOUBLE:
+		case EOSCTypeTag::OSC_DOUBLE:
 			Arguments.Add(FOSCType(Stream.ReadDouble()));
 			break;
-		case EOSCTypeTag::INT64:
+		case EOSCTypeTag::OSC_INT64:
 			Arguments.Add(FOSCType(Stream.ReadInt64()));
 			break;
-		case EOSCTypeTag::TRUE:
+		case EOSCTypeTag::OSC_TRUE:
 			Arguments.Add(FOSCType(true));
 			break;
-		case EOSCTypeTag::FALSE:
+		case EOSCTypeTag::OSC_FALSE:
 			Arguments.Add(FOSCType(false));
 			break;
-		case EOSCTypeTag::NIL:
-			Arguments.Add(FOSCType(EOSCTypeTag::NIL));
+		case EOSCTypeTag::OSC_NIL:
+			Arguments.Add(FOSCType(EOSCTypeTag::OSC_NIL));
 			break;
-		case EOSCTypeTag::INFINITUM:
-			Arguments.Add(FOSCType(EOSCTypeTag::INFINITUM));
+		case EOSCTypeTag::OSC_INFINITUM:
+			Arguments.Add(FOSCType(EOSCTypeTag::OSC_INFINITUM));
 			break;
-		case EOSCTypeTag::TIME:
+		case EOSCTypeTag::OSC_TIME:
 			Arguments.Add(FOSCType(Stream.ReadUInt64()));
 			break;
-		case EOSCTypeTag::STRING:
+		case EOSCTypeTag::OSC_STRING:
 			Arguments.Add(FOSCType(Stream.ReadString()));
 			break;
-		case EOSCTypeTag::BLOB:
+		case EOSCTypeTag::OSC_BLOB:
 			Arguments.Add(FOSCType(Stream.ReadBlob()));
 			break;
-		case EOSCTypeTag::COLOR:
+		case EOSCTypeTag::OSC_COLOR:
 			Arguments.Add(FOSCType(FColor(Stream.ReadInt32())));
 			break;
-		case EOSCTypeTag::TERMINATE:
+		case EOSCTypeTag::OSC_TERMINATE:
 			// Return on first terminate found. FString GetCharArray can return
 			// an array with multiple terminators.
 			return;
diff --git a/Engine/Plugins/Runtime/OSC/Source/OSC/Public/OSCTypes.h b/Engine/Plugins/Runtime/OSC/Source/OSC/Public/OSCTypes.h
--- a/Engine/Plugins/Runtime/OSC/Source/OSC/Public/OSCTypes.h
+++ b/Engine/Plugins/Runtime/OSC/Source/OSC/Public/OSCTypes.h
@@ -6,22 +6,22 @@
 #include "Math/Color.h"
 
 
-enum class EOSCTypeTag {
-
-	INT32 = 'i',
-	FLOAT = 'f',
-	DOUBLE = 'd',
-	STRING = 's',
-	BLOB = 'b',
-	TIME = 't',
-	INT64 = 'h',
-	CHAR = 'c',
-	TRUE = 'T',
-	FALSE = 'F',
-	NIL = 'N',
-	INFINITUM = 'I',
-	COLOR = 'r',
-	TERMINATE = '\0'
+enum class EOSCTypeTag : uint8
+{
+	OSC_INT32		= 'i',
+	OSC_FLOAT		= 'f',
+	OSC_DOUBLE		= 'd',
+	OSC_STRING		= 's',
+	OSC_BLOB		= 'b',
+	OSC_TIME		= 't',
+	OSC_INT64		= 'h',
+	OSC_CHAR		= 'c',
+	OSC_TRUE		= 'T',
+	OSC_FALSE		= 'F',
+	OSC_NIL			= 'N',
+	OSC_INFINITUM	= 'I',
+	OSC_COLOR		= 'r',
+	OSC_TERMINATE	= '\0'
 };
 
 
@@ -29,63 +29,63 @@ class FOSCType
 {
 public:
 	explicit FOSCType(int32 Value)
-		: TypeTag(EOSCTypeTag::INT32)
+		: TypeTag(EOSCTypeTag::OSC_INT32)
 		, Data(Value)
 		, Blob()
 	{
 	}
 	explicit FOSCType(int64 Value)
-		: TypeTag(EOSCTypeTag::INT64)
+		: TypeTag(EOSCTypeTag::OSC_INT64)
 		, Data(Value)
 		, Blob()
 	{
 	}
 	explicit FOSCType(ANSICHAR Value)
-		: TypeTag(EOSCTypeTag::CHAR)
+		: TypeTag(EOSCTypeTag::OSC_CHAR)
 		, Data(Value)
 		, Blob()
 	{
 	}
 	explicit FOSCType(uint64 Value)
-		: TypeTag(EOSCTypeTag::TIME)
+		: TypeTag(EOSCTypeTag::OSC_TIME)
 		, Data(Value)
 		, Blob()
 	{
 	}
 	explicit FOSCType(float Value)
-		: TypeTag(EOSCTypeTag::FLOAT)
+		: TypeTag(EOSCTypeTag::OSC_FLOAT)
 		, Data(Value)
 		, Blob()
 	
 	{
 	}
 	explicit FOSCType(double Value)
-		: TypeTag(EOSCTypeTag::DOUBLE)
+		: TypeTag(EOSCTypeTag::OSC_DOUBLE)
 		, Data(Value)
 		, Blob()
 	{
 	}
 	explicit FOSCType(bool Value)
-		: TypeTag(Value ? EOSCTypeTag::TRUE : EOSCTypeTag::FALSE)
+		: TypeTag(Value ? EOSCTypeTag::OSC_TRUE : EOSCTypeTag::OSC_FALSE)
 		, Data(Value)
 		, Blob()
 	{
 	}
 	explicit FOSCType(const FString& Value)
-		: TypeTag(EOSCTypeTag::STRING)
+		: TypeTag(EOSCTypeTag::OSC_STRING)
 		, Data(0)
 		, String(Value)
 		, Blob()
 	{
 	}
 	explicit FOSCType(const TArray<uint8>& Value)
-		: TypeTag(EOSCTypeTag::BLOB)
+		: TypeTag(EOSCTypeTag::OSC_BLOB)
 		, Data(0)
 		, Blob(Value)
 	{
 	}
 	explicit FOSCType(FColor Value)
-		: TypeTag(EOSCTypeTag::COLOR)
+		: TypeTag(EOSCTypeTag::OSC_COLOR)
 		, Data(0)
 		, Blob()
 		, Color(Value)
@@ -101,38 +101,38 @@ public:
 
 	EOSCTypeTag GetTypeTag() const { return TypeTag; }
 
-	bool IsInt32() const { return TypeTag == EOSCTypeTag::INT32; }
+	bool IsInt32() const { return TypeTag == EOSCTypeTag::OSC_INT32; }
 	int32 GetInt32() const { return Data.Int32; }
 
-	bool IsInt64() const { return TypeTag == EOSCTypeTag::INT64; }
+	bool IsInt64() const { return TypeTag == EOSCTypeTag::OSC_INT64; }
 	int64 GetInt64() const { return Data.Int64; }
 
-	bool IsTimeTag() const { return TypeTag == EOSCTypeTag::TIME; }
+	bool IsTimeTag() const { return TypeTag == EOSCTypeTag::OSC_TIME; }
 	uint64 GetTimeTag() const { return Data.Time; }
 
-	bool IsBool() const { return TypeTag == EOSCTypeTag::TRUE || TypeTag == EOSCTypeTag::FALSE; }
+	bool IsBool() const { return TypeTag == EOSCTypeTag::OSC_TRUE || TypeTag == EOSCTypeTag::OSC_FALSE; }
 	bool GetBool() const { return Data.Bool; }
 
-	bool IsChar() const { return TypeTag == EOSCTypeTag::CHAR; }
+	bool IsChar() const { return TypeTag == EOSCTypeTag::OSC_CHAR; }
 	ANSICHAR GetChar() const { return Data.Char; }
 
-	bool IsFloat() const { return TypeTag == EOSCTypeTag::FLOAT; }
+	bool IsFloat() const { return TypeTag == EOSCTypeTag::OSC_FLOAT; }
 	float GetFloat() const { return Data.Float; }
 	
-	bool IsDouble() const { return TypeTag == EOSCTypeTag::DOUBLE; }
+	bool IsDouble() const { return TypeTag == EOSCTypeTag::OSC_DOUBLE; }
 	double GetDouble() const { return Data.Double; }
 
-	bool IsString() const { return TypeTag == EOSCTypeTag::STRING; }
+	bool IsString() const { return TypeTag == EOSCTypeTag::OSC_STRING; }
 	FString GetString() const { return String; }
 
-	bool IsBlob() const { return TypeTag == EOSCTypeTag::BLOB; }
+	bool IsBlob() const { return TypeTag == EOSCTypeTag::OSC_BLOB; }
 	TArray<uint8> GetBlob() const { return Blob; }
 
-	bool IsColor() const { return TypeTag == EOSCTypeTag::COLOR; }
+	bool IsColor() const { return TypeTag == EOSCTypeTag::OSC_COLOR; }
 	FColor GetColor()  const { return Color; }
 
-	bool IsNil() const { return TypeTag == EOSCTypeTag::NIL; }
-	bool IsInfinitum() const { return TypeTag == EOSCTypeTag::INFINITUM; }
+	bool IsNil() const { return TypeTag == EOSCTypeTag::OSC_NIL; }
+	bool IsInfinitum() const { return TypeTag == EOSCTypeTag::OSC_INFINITUM; }
 
 private:
 	union DataTypes
