Fix logging line breaks & minor formatting in OSC cvars

diff --git a/Engine/Plugins/Runtime/OSC/Source/OSC/Private/OSCManager.cpp b/Engine/Plugins/Runtime/OSC/Source/OSC/Private/OSCManager.cpp
--- a/Engine/Plugins/Runtime/OSC/Source/OSC/Private/OSCManager.cpp
+++ b/Engine/Plugins/Runtime/OSC/Source/OSC/Private/OSCManager.cpp
@@ -98,10 +98,10 @@ static FAutoConsoleCommand GOSCPrintServers(
 					const TArray<FOSCAddress> BoundPatterns = Server->GetBoundOSCAddressPatterns();
 					if (BoundPatterns.Num() > 0)
 					{
-						UE_LOG(LogOSC, Display, TEXT("\n    Bound Address Patterns:"));
+						UE_LOG(LogOSC, Display, TEXT("    Bound Address Patterns:"));
 						for (const FOSCAddress& Pattern : BoundPatterns)
 						{
-							UE_LOG(LogOSC, Display, TEXT("\n         %s"), *Pattern.GetFullPath());
+							UE_LOG(LogOSC, Display, TEXT("         %s"), *Pattern.GetFullPath());
 						}
 						UE_LOG(LogOSC, Display, TEXT(""));
 					}
@@ -115,8 +115,8 @@ static FAutoConsoleCommand GOSCServerConnect(
 	TEXT("osc.server.connect"),
 	TEXT("Connects or reconnects the osc mix server with the provided name\n"
 		"(see \"osc.servers\" for a list of available servers and their respective names). Args:\n"
-		"Name - Object name of server to (re)connect"
-		"Address - IP Address to connect to (default: LocalHost)"
+		"Name - Object name of server to (re)connect\n"
+		"Address - IP Address to connect to (default: LocalHost)\n"
 		"Port - Port to connect to (default: 8095)"),
 	FConsoleCommandWithArgsDelegate::CreateStatic(
 		[](const TArray<FString>& Args)
@@ -165,8 +165,8 @@ static FAutoConsoleCommand GOSCServerConnectById(
 	TEXT("osc.server.connectById"),
 	TEXT("Connects or reconnects the osc mix server with the provided object id\n"
 		"(see \"osc.servers\" for a list of available servers and their respective ids). Args:\n"
-		"Id - Object Id of client to (re)connect"
-		"Address - IP Address to (re)connect to (default: LocalHost)"
+		"Id - Object Id of client to (re)connect\n"
+		"Address - IP Address to (re)connect to (default: LocalHost)\n"
 		"Port - Port to (re)connect to (default: 8095)"),
 	FConsoleCommandWithArgsDelegate::CreateStatic(
 		[](const TArray<FString>& Args)
@@ -237,9 +237,9 @@ static FAutoConsoleCommand GOSCPrintClients(
 					FString IPAddrStr;
 					int32 Port;
 					Client->GetSendIPAddress(IPAddrStr, Port);
-					ToPrint += TEXT(", ") + IPAddrStr + TEXT(":");
+					ToPrint += TEXT(", Send IP: ") + IPAddrStr + TEXT(":");
 					ToPrint.AppendInt(Port);
-					ToPrint += Client->IsActive() ? TEXT(" [Active]") : TEXT(" [Inactive]");
+					ToPrint += Client->IsActive() ? TEXT(") [Active]") : TEXT(") [Inactive]");
 
 					UE_LOG(LogOSC, Display, TEXT("%s"), *ToPrint);
 				}
@@ -252,8 +252,8 @@ static FAutoConsoleCommand GOSCClientConnect(
 	TEXT("osc.client.connect"),
 	TEXT("Connects (or reconnects) the osc mix client with the provided name\n"
 		"(see \"osc.clients\" for a list of available clients and their respective ids). Args:\n"
-		"Name - Object name of client to (re)connect"
-		"Address - IP Address to (re)connect to (default: LocalHost)"
+		"Name - Object name of client to (re)connect\n"
+		"Address - IP Address to (re)connect to (default: LocalHost)\n"
 		"Port - Port to (re)connect to (default: 8094)"),
 	FConsoleCommandWithArgsDelegate::CreateStatic(
 		[](const TArray< FString >& Args)
@@ -299,8 +299,8 @@ static FAutoConsoleCommand GOSCClientConnectById(
 	TEXT("osc.client.connectById"),
 	TEXT("Connects (or reconnects) the osc mix client with the provided object id\n"
 		"(see \"osc.clients\" for a list of available clients and their respective ids). Args:\n"
-		"Id - Object Id of client to (re)connect"
-		"Address - IP Address to (re)connect to (default: LocalHost)"
+		"Id - Object Id of client to (re)connect\n"
+		"Address - IP Address to (re)connect to (default: LocalHost)\n"
 		"Port - Port to (re)connect to (default: 8094)"),
 	FConsoleCommandWithArgsDelegate::CreateStatic(
 		[](const TArray< FString >& Args)
