Move OSCServerProxy packet processing to be pumped on tick and not forwarding request to task graph.  Avoids crash caused by callback on destructed server proxy.
#test -game, editor, stand-alone client in QAGame, stress test spamming messages while starting and stopping OSCServer

diff --git a/Engine/Plugins/Runtime/OSC/Source/OSC/Private/OSCBundlePacket.cpp b/Engine/Plugins/Runtime/OSC/Source/OSC/Private/OSCBundlePacket.cpp
--- a/Engine/Plugins/Runtime/OSC/Source/OSC/Private/OSCBundlePacket.cpp
+++ b/Engine/Plugins/Runtime/OSC/Source/OSC/Private/OSCBundlePacket.cpp
@@ -66,7 +66,7 @@ void FOSCBundlePacket::ReadData(FOSCStream& Stream)
 		int32 PacketLength = Stream.ReadInt32();
 
 		int32 StartPos = Stream.GetPosition();
-		TSharedPtr<IOSCPacket> Packet = IOSCPacket::CreatePacket(Stream.GetData() + Stream.GetPosition());
+		TSharedPtr<IOSCPacket> Packet = IOSCPacket::CreatePacket(Stream.GetData() + Stream.GetPosition(), Endpoint);
 		if (!Packet.IsValid())
 		{
 			break;
diff --git a/Engine/Plugins/Runtime/OSC/Source/OSC/Private/OSCManager.cpp b/Engine/Plugins/Runtime/OSC/Source/OSC/Private/OSCManager.cpp
--- a/Engine/Plugins/Runtime/OSC/Source/OSC/Private/OSCManager.cpp
+++ b/Engine/Plugins/Runtime/OSC/Source/OSC/Private/OSCManager.cpp
@@ -356,7 +356,6 @@ UOSCServer* UOSCManager::CreateOSCServer(FString InReceiveIPAddress, int32 InPor
 
 	if (UOSCServer* NewOSCServer = NewObject<UOSCServer>(GetTransientPackage(), FName(*ServerName)))
 	{
-		NewOSCServer->Connect();
 		NewOSCServer->SetMulticastLoopback(bInMulticastLoopback);
 		if (NewOSCServer->SetAddress(InReceiveIPAddress, InPort))
 		{
diff --git a/Engine/Plugins/Runtime/OSC/Source/OSC/Private/OSCPacket.cpp b/Engine/Plugins/Runtime/OSC/Source/OSC/Private/OSCPacket.cpp
--- a/Engine/Plugins/Runtime/OSC/Source/OSC/Private/OSCPacket.cpp
+++ b/Engine/Plugins/Runtime/OSC/Source/OSC/Private/OSCPacket.cpp
@@ -1,29 +1,25 @@
 // Copyright Epic Games, Inc. All Rights Reserved.
 #include "OSCPacket.h"
 
+#include "Interfaces/IPv4/IPv4Endpoint.h"
+
 #include "OSCMessagePacket.h"
 #include "OSCBundlePacket.h"
 #include "OSCLog.h"
 
 
-IOSCPacket::IOSCPacket()
-{
-}
-
-IOSCPacket::~IOSCPacket()
-{
-}
-
-TSharedPtr<IOSCPacket> IOSCPacket::CreatePacket(const uint8* PacketType)
+TSharedPtr<IOSCPacket> IOSCPacket::CreatePacket(const uint8* InPacketType, const FIPv4Endpoint& InEndpoint)
 {
-	const FString PacketIdentifier(ANSI_TO_TCHAR((const ANSICHAR*)&PacketType[0]));
+	const FString PacketIdentifier(ANSI_TO_TCHAR((const ANSICHAR*)&InPacketType[0]));
+	
+	TSharedPtr<IOSCPacket> Packet;
 	if (PacketIdentifier.StartsWith(OSC::PathSeparator))
 	{
-		return MakeShareable(new FOSCMessagePacket());
+		Packet = MakeShared<FOSCMessagePacket>();
 	}
 	else if (PacketIdentifier == OSC::BundleTag)
 	{
-		return MakeShareable(new FOSCBundlePacket());
+		Packet = MakeShared<FOSCBundlePacket>();
 	}
 	else
 	{
@@ -31,4 +27,12 @@ TSharedPtr<IOSCPacket> IOSCPacket::CreatePacket(const uint8* PacketType)
 			"Lead identifier of '%c' not valid bundle tag ('%s') or message ('%s') identifier."), *PacketIdentifier, *OSC::BundleTag, *OSC::PathSeparator);
 		return nullptr;
 	}
+
+	Packet->Endpoint = InEndpoint;
+	return Packet;
 }
+
+const FIPv4Endpoint& IOSCPacket::GetEndpoint() const
+{
+	return Endpoint;
+}
\ No newline at end of file
diff --git a/Engine/Plugins/Runtime/OSC/Source/OSC/Private/OSCServer.cpp b/Engine/Plugins/Runtime/OSC/Source/OSC/Private/OSCServer.cpp
--- a/Engine/Plugins/Runtime/OSC/Source/OSC/Private/OSCServer.cpp
+++ b/Engine/Plugins/Runtime/OSC/Source/OSC/Private/OSCServer.cpp
@@ -17,15 +17,10 @@
 
 UOSCServer::UOSCServer(const FObjectInitializer& ObjectInitializer)
 	: Super(ObjectInitializer)
-	, ServerProxy(new FOSCServerProxy(*this))
+	, ServerProxy(MakeUnique<FOSCServerProxy>(*this))
 {
 }
 
-void UOSCServer::Connect()
-{
-	ServerProxy.Reset(new FOSCServerProxy(*this));
-}
-
 bool UOSCServer::GetMulticastLoopback() const
 {
 	check(ServerProxy.IsValid());
@@ -56,6 +51,15 @@ void UOSCServer::SetMulticastLoopback(bool bInMulticastLoopback)
 	ServerProxy->SetMulticastLoopback(bInMulticastLoopback);
 }
 
+#if WITH_EDITOR
+void UOSCServer::SetTickInEditor(bool bInTickInEditor)
+{
+	check(ServerProxy.IsValid());
+	ServerProxy->SetTickableInEditor(bInTickInEditor);
+}
+#endif // WITH_EDITOR
+
+
 void UOSCServer::Stop()
 {
 	if (ServerProxy.IsValid())
@@ -172,15 +176,15 @@ void UOSCServer::ClearPackets()
 	OSCPackets.Empty();
 }
 
-void UOSCServer::EnqueuePacket(TSharedPtr<IOSCPacket> Packet)
+void UOSCServer::EnqueuePacket(TSharedPtr<IOSCPacket> InPacket)
 {
-	OSCPackets.Enqueue(Packet);
+	OSCPackets.Enqueue(InPacket);
 }
 
-void UOSCServer::DispatchBundle(const FString& InIPAddress, uint16 Port, const FOSCBundle& InBundle)
+void UOSCServer::DispatchBundle(const FString& InIPAddress, uint16 InPort, const FOSCBundle& InBundle)
 {
-	OnOscBundleReceived.Broadcast(InBundle, InIPAddress, Port);
-	OnOscBundleReceivedNative.Broadcast(InBundle, InIPAddress, Port);
+	OnOscBundleReceived.Broadcast(InBundle, InIPAddress, InPort);
+	OnOscBundleReceivedNative.Broadcast(InBundle, InIPAddress, InPort);
 
 	TSharedPtr<FOSCBundlePacket> BundlePacket = StaticCastSharedPtr<FOSCBundlePacket>(InBundle.GetPacket());
 	FOSCBundlePacket::FPacketBundle Packets = BundlePacket->GetPackets();
@@ -188,11 +192,11 @@ void UOSCServer::DispatchBundle(const FString& InIPAddress, uint16 Port, const F
 	{
 		if (Packet->IsMessage())
 		{
-			DispatchMessage(InIPAddress, Port, FOSCMessage(Packet));
+			DispatchMessage(InIPAddress, InPort, FOSCMessage(Packet));
 		}
 		else if (Packet->IsBundle())
 		{
-			DispatchBundle(InIPAddress, Port, FOSCBundle(Packet));
+			DispatchBundle(InIPAddress, InPort, FOSCBundle(Packet));
 		}
 		else
 		{
@@ -201,10 +205,10 @@ void UOSCServer::DispatchBundle(const FString& InIPAddress, uint16 Port, const F
 	}
 }
 
-void UOSCServer::DispatchMessage(const FString& InIPAddress, uint16 Port, const FOSCMessage& InMessage)
+void UOSCServer::DispatchMessage(const FString& InIPAddress, uint16 InPort, const FOSCMessage& InMessage)
 {
-	OnOscMessageReceived.Broadcast(InMessage, InIPAddress, Port);
-	OnOscMessageReceivedNative.Broadcast(InMessage, InIPAddress, Port);
+	OnOscMessageReceived.Broadcast(InMessage, InIPAddress, InPort);
+	OnOscMessageReceivedNative.Broadcast(InMessage, InIPAddress, InPort);
 
 	UE_LOG(LogOSC, Verbose, TEXT("Message received from endpoint '%s', OSCAddress of '%s'."), *InIPAddress, *InMessage.GetAddress().GetFullPath());
 
@@ -213,7 +217,7 @@ void UOSCServer::DispatchMessage(const FString& InIPAddress, uint16 Port, const
 		const FOSCDispatchMessageEvent& DispatchEvent = Pair.Value;
 		if (Pair.Key.Matches(InMessage.GetAddress()))
 		{
-			DispatchEvent.Broadcast(Pair.Key, InMessage, InIPAddress, Port);
+			DispatchEvent.Broadcast(Pair.Key, InMessage, InIPAddress, InPort);
 			UE_LOG(LogOSC, Verbose, TEXT("Message dispatched from endpoint '%s', OSCAddress path of '%s' matched OSCAddress pattern '%s'."),
 				*InIPAddress,
 				*InMessage.GetAddress().GetFullPath(),
@@ -222,22 +226,26 @@ void UOSCServer::DispatchMessage(const FString& InIPAddress, uint16 Port, const
 	}
 }
 
-void UOSCServer::OnPacketReceived(const FString& InIPAddress, uint16 Port)
+void UOSCServer::PumpPacketQueue(const TSet<uint32>* WhitelistedClients)
 {
 	TSharedPtr<IOSCPacket> Packet;
 	while (OSCPackets.Dequeue(Packet))
 	{
-		if (Packet->IsMessage())
-		{
-			DispatchMessage(InIPAddress, Port, FOSCMessage(Packet));
-		}
-		else if (Packet->IsBundle())
-		{
-			DispatchBundle(InIPAddress, Port, FOSCBundle(Packet));
-		}
-		else
+		const FIPv4Endpoint& Endpoint = Packet->GetEndpoint();
+		if (!WhitelistedClients || WhitelistedClients->Contains(Endpoint.Address.Value))
 		{
-			UE_LOG(LogOSC, Warning, TEXT("Failed to parse invalid received message. Invalid OSC type (packet is neither identified as message nor bundle)."));
+			if (Packet->IsMessage())
+			{
+				DispatchMessage(Endpoint.Address.ToString(), Endpoint.Port, FOSCMessage(Packet));
+			}
+			else if (Packet->IsBundle())
+			{
+				DispatchBundle(Endpoint.Address.ToString(), Endpoint.Port, FOSCBundle(Packet));
+			}
+			else
+			{
+				UE_LOG(LogOSC, Warning, TEXT("Failed to parse invalid received message. Invalid OSC type (packet is neither identified as message nor bundle)."));
+			}
 		}
 	}
 }
\ No newline at end of file
diff --git a/Engine/Plugins/Runtime/OSC/Source/OSC/Private/OSCServerProxy.cpp b/Engine/Plugins/Runtime/OSC/Source/OSC/Private/OSCServerProxy.cpp
--- a/Engine/Plugins/Runtime/OSC/Source/OSC/Private/OSCServerProxy.cpp
+++ b/Engine/Plugins/Runtime/OSC/Source/OSC/Private/OSCServerProxy.cpp
@@ -2,8 +2,10 @@
 #include "OSCServerProxy.h"
 
 #include "Common/UdpSocketBuilder.h"
-#include "Runtime/Core/Public/Async/TaskGraphInterfaces.h"
+#include "CoreGlobals.h"
 #include "Sockets.h"
+#include "Stats/Stats.h"
+#include "Tickable.h"
 
 #include "OSCLog.h"
 #include "OSCStream.h"
@@ -17,6 +19,9 @@ FOSCServerProxy::FOSCServerProxy(UOSCServer& InServer)
 	, Port(0)
 	, bMulticastLoopback(false)
 	, bWhitelistClients(false)
+#if WITH_EDITOR
+	, bTickInEditor(false)
+#endif // WITH_EDITOR
 {
 }
 
@@ -25,31 +30,18 @@ FOSCServerProxy::~FOSCServerProxy()
 	Stop();
 }
 
-void FOSCServerProxy::OnPacketReceived(const FArrayReaderPtr& Data, const FIPv4Endpoint& Endpoint)
+void FOSCServerProxy::OnPacketReceived(const FArrayReaderPtr& InData, const FIPv4Endpoint& InEndpoint)
 {
-	TSharedPtr<IOSCPacket> Packet = IOSCPacket::CreatePacket(Data->GetData());
+	TSharedPtr<IOSCPacket> Packet = IOSCPacket::CreatePacket(InData->GetData(), InEndpoint);
 	if (!Packet.IsValid())
 	{
-		UE_LOG(LogOSC, Verbose, TEXT("Message received from endpoint '%s' invalid OSC packet."), *Endpoint.ToString());
+		UE_LOG(LogOSC, Verbose, TEXT("Message received from endpoint '%s' invalid OSC packet."), *InEndpoint.ToString());
 		return;
 	}
 
-	FOSCStream Stream = FOSCStream(Data->GetData(), Data->Num());
+	FOSCStream Stream = FOSCStream(InData->GetData(), InData->Num());
 	Packet->ReadData(Stream);
 	Server->EnqueuePacket(Packet);
-
-	DECLARE_CYCLE_STAT(TEXT("OSCServer.OnPacketReceived"), STAT_OSCServerOnPacketReceived, STATGROUP_OSCNetworkCommands);
-	FFunctionGraphTask::CreateAndDispatchWhenReady([this, Endpoint]()
-	{
-		// Throw request on the ground if endpoint address not whitelisted.
-		if (bWhitelistClients && !ClientWhitelist.Contains(Endpoint.Address.Value))
-		{
-			Server->ClearPackets();
-			return;
-		}
-
-		Server->OnPacketReceived(Endpoint.Address.ToString(), Endpoint.Port);
-	}, GET_STATID(STAT_OSCServerOnPacketReceived), nullptr, ENamedThreads::GameThread);
 }
 
 FString FOSCServerProxy::GetIpAddress() const
@@ -72,16 +64,16 @@ bool FOSCServerProxy::IsActive() const
 	return SocketReceiver != nullptr;
 }
 
-void FOSCServerProxy::Listen(const FString& ServerName)
+void FOSCServerProxy::Listen(const FString& InServerName)
 {
 	if (IsActive())
 	{
 		UE_LOG(LogOSC, Error, TEXT("OSCServer currently listening: %s:%d. Failed to start new service prior to calling stop."),
-			*ServerName, *ReceiveIPAddress.ToString(), Port);
+			*InServerName, *ReceiveIPAddress.ToString(), Port);
 		return;
 	}
 
-	FUdpSocketBuilder Builder(*ServerName);
+	FUdpSocketBuilder Builder(*InServerName);
 	Builder.BoundToPort(Port);
 	if (ReceiveIPAddress.IsMulticastAddress())
 	{
@@ -96,7 +88,7 @@ void FOSCServerProxy::Listen(const FString& ServerName)
 		if (bMulticastLoopback)
 		{
 			UE_LOG(LogOSC, Warning, TEXT("OSCServer '%s' ReceiveIPAddress provided is not a multicast address.  Not respecting MulticastLoopback boolean."),
-				*ServerName);
+				*InServerName);
 		}
 		Builder.BoundToAddress(ReceiveIPAddress);
 	}
@@ -104,15 +96,15 @@ void FOSCServerProxy::Listen(const FString& ServerName)
 	Socket = Builder.Build();
 	if (Socket)
 	{
-		SocketReceiver = new FUdpSocketReceiver(Socket, FTimespan::FromMilliseconds(100), *(ServerName + TEXT("_ListenerThread")));
+		SocketReceiver = new FUdpSocketReceiver(Socket, FTimespan::FromMilliseconds(100), *(InServerName + TEXT("_ListenerThread")));
 		SocketReceiver->OnDataReceived().BindRaw(this, &FOSCServerProxy::OnPacketReceived);
 		SocketReceiver->Start();
 
-		UE_LOG(LogOSC, Display, TEXT("OSCServer '%s' Listening: %s:%d."), *ServerName, *ReceiveIPAddress.ToString(), Port);
+		UE_LOG(LogOSC, Display, TEXT("OSCServer '%s' Listening: %s:%d."), *InServerName, *ReceiveIPAddress.ToString(), Port);
 	}
 	else
 	{
-		UE_LOG(LogOSC, Warning, TEXT("OSCServer '%s' failed to bind to socket on %s:%d."), *ServerName, *ReceiveIPAddress.ToString(), Port);
+		UE_LOG(LogOSC, Warning, TEXT("OSCServer '%s' failed to bind to socket on %s:%d."), *InServerName, *ReceiveIPAddress.ToString(), Port);
 	}
 }
 
@@ -145,6 +137,18 @@ void FOSCServerProxy::SetMulticastLoopback(bool bInMulticastLoopback)
 	bMulticastLoopback = bInMulticastLoopback;
 }
 
+#if WITH_EDITOR
+bool FOSCServerProxy::IsTickableInEditor() const
+{
+	return bTickInEditor;
+}
+
+void FOSCServerProxy::SetTickableInEditor(bool bInTickInEditor)
+{
+	bTickInEditor = bInTickInEditor;
+}
+#endif // WITH_EDITOR
+
 void FOSCServerProxy::Stop()
 {
 	if (SocketReceiver)
@@ -159,26 +163,28 @@ void FOSCServerProxy::Stop()
 		ISocketSubsystem::Get(PLATFORM_SOCKETSUBSYSTEM)->DestroySocket(Socket);
 		Socket = nullptr;
 	}
+
+
 }
 
-void FOSCServerProxy::AddWhitelistedClient(const FString& IPAddress)
+void FOSCServerProxy::AddWhitelistedClient(const FString& InIPAddress)
 {
 	FIPv4Address OutAddress;
-	if (!FIPv4Address::Parse(IPAddress, OutAddress))
+	if (!FIPv4Address::Parse(InIPAddress, OutAddress))
 	{
-		UE_LOG(LogOSC, Warning, TEXT("OSCServer '%s' failed to whitelist IP Address '%s'. Address is invalid."), *IPAddress);
+		UE_LOG(LogOSC, Warning, TEXT("OSCServer '%s' failed to whitelist IP Address '%s'. Address is invalid."), *InIPAddress);
 		return;
 	}
 
 	ClientWhitelist.Add(OutAddress.Value);
 }
 
-void FOSCServerProxy::RemoveWhitelistedClient(const FString& IPAddress)
+void FOSCServerProxy::RemoveWhitelistedClient(const FString& InIPAddress)
 {
 	FIPv4Address OutAddress;
-	if (!FIPv4Address::Parse(IPAddress, OutAddress))
+	if (!FIPv4Address::Parse(InIPAddress, OutAddress))
 	{
-		UE_LOG(LogOSC, Warning, TEXT("OSCServer '%s' failed to remove whitelisted IP Address '%s'. Address is invalid."), *IPAddress);
+		UE_LOG(LogOSC, Warning, TEXT("OSCServer '%s' failed to remove whitelisted IP Address '%s'. Address is invalid."), *InIPAddress);
 		return;
 	}
 
@@ -201,7 +207,26 @@ TSet<FString> FOSCServerProxy::GetWhitelistedClients() const
 	return OutWhitelist;
 }
 
-void FOSCServerProxy::SetWhitelistClientsEnabled(bool bEnabled)
+void FOSCServerProxy::SetWhitelistClientsEnabled(bool bInEnabled)
+{
+	bWhitelistClients = bInEnabled;
+}
+
+void FOSCServerProxy::Tick(float InDeltaTime)
+{
+	check(IsInGameThread());
+	check(Server);
+
+	Server->PumpPacketQueue(bWhitelistClients ? &ClientWhitelist : nullptr);
+}
+
+TStatId FOSCServerProxy::GetStatId() const
+{
+	RETURN_QUICK_DECLARE_CYCLE_STAT(FOSCServerProxy, STATGROUP_Tickables);
+}
+
+UWorld* FOSCServerProxy::GetTickableGameObjectWorld() const
 {
-	bWhitelistClients = bEnabled;
+	check(Server);
+	return Server->GetWorld();
 }
diff --git a/Engine/Plugins/Runtime/OSC/Source/OSC/Private/OSCServerProxy.h b/Engine/Plugins/Runtime/OSC/Source/OSC/Private/OSCServerProxy.h
--- a/Engine/Plugins/Runtime/OSC/Source/OSC/Private/OSCServerProxy.h
+++ b/Engine/Plugins/Runtime/OSC/Source/OSC/Private/OSCServerProxy.h
@@ -4,41 +4,61 @@
 #include "CoreMinimal.h"
 
 #include "Common/UdpSocketReceiver.h"
+#include "HAL/ThreadSafeBool.h"
 #include "Interfaces/IPv4/IPv4Address.h"
 #include "Interfaces/IPv4/IPv4Endpoint.h"
+#include "Tickable.h"
 #include "UObject/Object.h"
 
 #include "OSCServer.h"
 
 
-class OSC_API FOSCServerProxy : public IOSCServerProxy
+class OSC_API FOSCServerProxy : public IOSCServerProxy, public FTickableGameObject
 {
 public:
 	FOSCServerProxy(UOSCServer& InServer);
 	virtual ~FOSCServerProxy();
 
+	// Begin IOSCServerProxy interface
 	bool GetMulticastLoopback() const override;
 	FString GetIpAddress() const override;
 	int32 GetPort() const override;
 	bool IsActive() const override;
 
-	void Listen(const FString& ServerName) override;
+	void Listen(const FString& InServerName) override;
 
-	bool SetAddress(const FString& ReceiveIPAddress, int32 Port) override;
-	void SetMulticastLoopback(bool bMulticastLoopback) override;
+	bool SetAddress(const FString& InReceiveIPAddress, int32 InPort) override;
+	void SetMulticastLoopback(bool bInMulticastLoopback) override;
 
 	void Stop() override;
 
-	void AddWhitelistedClient(const FString& IPAddress) override;
-	void RemoveWhitelistedClient(const FString& IPAddress) override;
+	void AddWhitelistedClient(const FString& InIPAddress) override;
+	void RemoveWhitelistedClient(const FString& InIPAddress) override;
 	void ClearWhitelistedClients() override;
 	TSet<FString> GetWhitelistedClients() const override;
 
-	void SetWhitelistClientsEnabled(bool bEnabled) override;
+	void SetWhitelistClientsEnabled(bool bInEnabled) override;
+
+#if WITH_EDITOR
+	void SetTickableInEditor(bool bInTickInEditor) override;
+#endif // WITH_EDITOR
+	// End IOSCServerProxy interface
 
 	/** Callback that receives data from a socket. */
-	void OnPacketReceived(const FArrayReaderPtr& Data, const FIPv4Endpoint& Endpoint);
+	void OnPacketReceived(const FArrayReaderPtr& InData, const FIPv4Endpoint& InEndpoint);
+
+	// Begin FTickableGameObject interface
+	virtual void Tick(float InDeltaTime) override;
+	virtual bool IsTickable() const override { return IsActive(); }
+	virtual TStatId GetStatId() const override;
+	virtual UWorld* GetTickableGameObjectWorld() const override;
+
+#if WITH_EDITOR
+	virtual bool IsTickableInEditor() const override;
+#endif // WITH_EDITOR
+	// End FTickableGameObject interface
 
+	/** Parent server UObject */
 	UOSCServer* Server;
 
 	/** Socket used to listen for OSC packets. */
@@ -61,4 +81,9 @@ public:
 
 	/** Whether or not to use client whitelist */
 	bool bWhitelistClients;
+
+#if WITH_EDITOR
+	bool bTickInEditor;
+#endif // WITH_EDITOR
+
 };
diff --git a/Engine/Plugins/Runtime/OSC/Source/OSC/Public/OSCPacket.h b/Engine/Plugins/Runtime/OSC/Source/OSC/Public/OSCPacket.h
--- a/Engine/Plugins/Runtime/OSC/Source/OSC/Public/OSCPacket.h
+++ b/Engine/Plugins/Runtime/OSC/Source/OSC/Public/OSCPacket.h
@@ -1,6 +1,7 @@
 // Copyright Epic Games, Inc. All Rights Reserved.
 #pragma once
 
+#include "Interfaces/IPv4/IPv4Endpoint.h"
 #include "OSCTypes.h"
 #include "OSCStream.h"
 #include "OSCAddress.h"
@@ -12,14 +13,14 @@ class FOSCStream;
 class IOSCPacket
 {
 public:
-	IOSCPacket();
-	virtual ~IOSCPacket();
+	IOSCPacket() = default;
+	virtual ~IOSCPacket() = default;
 
 	/** Write packet data into stream */
-	virtual void WriteData(FOSCStream& Stream) = 0;
+	virtual void WriteData(FOSCStream& OutStream) = 0;
 	
 	/** Read packet data from stream */
-	virtual void ReadData(FOSCStream& Stream) = 0;
+	virtual void ReadData(FOSCStream& OutStream) = 0;
 	
 	/** Returns true if packet is message */
 	virtual bool IsMessage() = 0;
@@ -27,6 +28,12 @@ public:
 	/** Returns true if packet is bundle */
 	virtual bool IsBundle() = 0;
 
+	/** Get endpoint responsible for creation/forwarding of packet */
+	virtual const FIPv4Endpoint& GetEndpoint() const;
+
 	/** Create an OSC packet according to the input data. */
-	static TSharedPtr<IOSCPacket> CreatePacket(const uint8* PacketType);
+	static TSharedPtr<IOSCPacket> CreatePacket(const uint8* InPacketType, const FIPv4Endpoint& InEndpoint);
+
+protected:
+	FIPv4Endpoint Endpoint;
 };
diff --git a/Engine/Plugins/Runtime/OSC/Source/OSC/Public/OSCServer.h b/Engine/Plugins/Runtime/OSC/Source/OSC/Public/OSCServer.h
--- a/Engine/Plugins/Runtime/OSC/Source/OSC/Public/OSCServer.h
+++ b/Engine/Plugins/Runtime/OSC/Source/OSC/Public/OSCServer.h
@@ -42,6 +42,9 @@ public:
 	virtual void Listen(const FString& ServerName) = 0;
 	virtual bool SetAddress(const FString& InReceiveIPAddress, int32 InPort) = 0;
 	virtual void SetMulticastLoopback(bool bInMulticastLoopback) = 0;
+#if WITH_EDITOR
+	virtual void SetTickableInEditor(bool bInTickInEditor) = 0;
+#endif // WITH_EDITOR
 	virtual void Stop() = 0;
 	virtual void AddWhitelistedClient(const FString& InIPAddress) = 0;
 	virtual void RemoveWhitelistedClient(const FString& IPAddress) = 0;
@@ -56,8 +59,6 @@ class OSC_API UOSCServer : public UObject
 	GENERATED_UCLASS_BODY()
 
 public:
-	void Connect();
-
 	/** Gets whether or not to loopback if ReceiveIPAddress provided is multicast. */
 	UFUNCTION(BlueprintCallable, Category = "Audio|OSC")
 	bool GetMulticastLoopback() const;
@@ -146,24 +147,32 @@ public:
 	UFUNCTION(BlueprintCallable, Category = "Audio|OSC")
 	TArray<FOSCAddress> GetBoundOSCAddressPatterns() const;
 
+#if WITH_EDITOR
+	/** Set whether server instance can be ticked in-editor (editor only and available to blueprint
+	  * for use in editor utility scripts/script actions).
+	  */
+	UFUNCTION(BlueprintCallable, Category = "Audio|OSC")
+	void SetTickInEditor(bool bInTickInEditor);
+#endif // WITH_EDITOR
+
 	/** Clears all packets pending processing */
 	void ClearPackets();
 
 	/** Enqueues packet to be processed */
-	void EnqueuePacket(TSharedPtr<IOSCPacket>);
+	void EnqueuePacket(TSharedPtr<IOSCPacket> InPacket);
 
 	/** Callback for when packet is received by server */
-	void OnPacketReceived(const FString& InIPAddress, uint16 Port);
+	void PumpPacketQueue(const TSet<uint32>* InWhitelistedClients);
 
 protected:
 	void BeginDestroy() override;
 
 private:
 	/** Dispatches provided bundle received */
-	void DispatchBundle(const FString& InIPAddress, uint16 Port, const FOSCBundle& InBundle);
+	void DispatchBundle(const FString& InIPAddress, uint16 InPort, const FOSCBundle& InBundle);
 
 	/** Dispatches provided message received */
-	void DispatchMessage(const FString& InIPAddress, uint16 Port, const FOSCMessage& InMessage);
+	void DispatchMessage(const FString& InIPAddress, uint16 InPort, const FOSCMessage& InMessage);
 
 	/** Pointer to internal implementation of server proxy */
 	TUniquePtr<IOSCServerProxy> ServerProxy;
