Fix CIS issue with OSC dependency on Networking not being public. Add missing _API identifiers to public API
#rnx

diff --git a/Engine/Plugins/Runtime/OSC/Source/OSC/Private/OSCBundlePacket.cpp b/Engine/Plugins/Runtime/OSC/Source/OSC/Private/OSCBundlePacket.cpp
--- a/Engine/Plugins/Runtime/OSC/Source/OSC/Private/OSCBundlePacket.cpp
+++ b/Engine/Plugins/Runtime/OSC/Source/OSC/Private/OSCBundlePacket.cpp
@@ -66,7 +66,7 @@ void FOSCBundlePacket::ReadData(FOSCStream& Stream)
 		int32 PacketLength = Stream.ReadInt32();
 
 		int32 StartPos = Stream.GetPosition();
-		TSharedPtr<IOSCPacket> Packet = IOSCPacket::CreatePacket(Stream.GetData() + Stream.GetPosition(), Endpoint);
+		TSharedPtr<IOSCPacket> Packet = IOSCPacket::CreatePacket(Stream.GetData() + Stream.GetPosition(), IPAddress, Port);
 		if (!Packet.IsValid())
 		{
 			break;
diff --git a/Engine/Plugins/Runtime/OSC/Source/OSC/Private/OSCPacket.cpp b/Engine/Plugins/Runtime/OSC/Source/OSC/Private/OSCPacket.cpp
--- a/Engine/Plugins/Runtime/OSC/Source/OSC/Private/OSCPacket.cpp
+++ b/Engine/Plugins/Runtime/OSC/Source/OSC/Private/OSCPacket.cpp
@@ -1,14 +1,12 @@
 // Copyright Epic Games, Inc. All Rights Reserved.
 #include "OSCPacket.h"
 
-#include "Interfaces/IPv4/IPv4Endpoint.h"
-
 #include "OSCMessagePacket.h"
 #include "OSCBundlePacket.h"
 #include "OSCLog.h"
 
 
-TSharedPtr<IOSCPacket> IOSCPacket::CreatePacket(const uint8* InPacketType, const FIPv4Endpoint& InEndpoint)
+TSharedPtr<IOSCPacket> IOSCPacket::CreatePacket(const uint8* InPacketType, const FString& InIPAddress, uint16 InPort)
 {
 	const FString PacketIdentifier(ANSI_TO_TCHAR((const ANSICHAR*)&InPacketType[0]));
 	
@@ -28,11 +26,17 @@ TSharedPtr<IOSCPacket> IOSCPacket::CreatePacket(const uint8* InPacketType, const
 		return nullptr;
 	}
 
-	Packet->Endpoint = InEndpoint;
+	Packet->IPAddress = InIPAddress;
+	Packet->Port = InPort;
 	return Packet;
 }
 
-const FIPv4Endpoint& IOSCPacket::GetEndpoint() const
+const FString& IOSCPacket::GetIPAddress() const
+{
+	return IPAddress;
+}
+
+uint16 IOSCPacket::GetPort() const
 {
-	return Endpoint;
+	return Port;
 }
\ No newline at end of file
diff --git a/Engine/Plugins/Runtime/OSC/Source/OSC/Private/OSCServer.cpp b/Engine/Plugins/Runtime/OSC/Source/OSC/Private/OSCServer.cpp
--- a/Engine/Plugins/Runtime/OSC/Source/OSC/Private/OSCServer.cpp
+++ b/Engine/Plugins/Runtime/OSC/Source/OSC/Private/OSCServer.cpp
@@ -1,10 +1,11 @@
 // Copyright Epic Games, Inc. All Rights Reserved.
 #include "OSCServer.h"
 
-#include "Runtime/Core/Public/Async/TaskGraphInterfaces.h"
-#include "Sockets.h"
 #include "Common/UdpSocketBuilder.h"
 #include "Common/UdpSocketReceiver.h"
+#include "Interfaces/IPv4/IPv4Address.h"
+#include "Runtime/Core/Public/Async/TaskGraphInterfaces.h"
+#include "Sockets.h"
 
 #include "OSCStream.h"
 #include "OSCMessage.h"
@@ -231,16 +232,18 @@ void UOSCServer::PumpPacketQueue(const TSet<uint32>* WhitelistedClients)
 	TSharedPtr<IOSCPacket> Packet;
 	while (OSCPackets.Dequeue(Packet))
 	{
-		const FIPv4Endpoint& Endpoint = Packet->GetEndpoint();
-		if (!WhitelistedClients || WhitelistedClients->Contains(Endpoint.Address.Value))
+		FIPv4Address IPAddr;
+		const FString& Address = Packet->GetIPAddress();
+		if (!WhitelistedClients || (FIPv4Address::Parse(Address, IPAddr) && WhitelistedClients->Contains(IPAddr.Value)))
 		{
+			uint16 Port = Packet->GetPort();
 			if (Packet->IsMessage())
 			{
-				DispatchMessage(Endpoint.Address.ToString(), Endpoint.Port, FOSCMessage(Packet));
+				DispatchMessage(Address, Port, FOSCMessage(Packet));
 			}
 			else if (Packet->IsBundle())
 			{
-				DispatchBundle(Endpoint.Address.ToString(), Endpoint.Port, FOSCBundle(Packet));
+				DispatchBundle(Address, Port, FOSCBundle(Packet));
 			}
 			else
 			{
diff --git a/Engine/Plugins/Runtime/OSC/Source/OSC/Private/OSCServerProxy.cpp b/Engine/Plugins/Runtime/OSC/Source/OSC/Private/OSCServerProxy.cpp
--- a/Engine/Plugins/Runtime/OSC/Source/OSC/Private/OSCServerProxy.cpp
+++ b/Engine/Plugins/Runtime/OSC/Source/OSC/Private/OSCServerProxy.cpp
@@ -32,7 +32,7 @@ FOSCServerProxy::~FOSCServerProxy()
 
 void FOSCServerProxy::OnPacketReceived(const FArrayReaderPtr& InData, const FIPv4Endpoint& InEndpoint)
 {
-	TSharedPtr<IOSCPacket> Packet = IOSCPacket::CreatePacket(InData->GetData(), InEndpoint);
+	TSharedPtr<IOSCPacket> Packet = IOSCPacket::CreatePacket(InData->GetData(), InEndpoint.Address.ToString(), InEndpoint.Port);
 	if (!Packet.IsValid())
 	{
 		UE_LOG(LogOSC, Verbose, TEXT("Message received from endpoint '%s' invalid OSC packet."), *InEndpoint.ToString());
diff --git a/Engine/Plugins/Runtime/OSC/Source/OSC/Public/OSCPacket.h b/Engine/Plugins/Runtime/OSC/Source/OSC/Public/OSCPacket.h
--- a/Engine/Plugins/Runtime/OSC/Source/OSC/Public/OSCPacket.h
+++ b/Engine/Plugins/Runtime/OSC/Source/OSC/Public/OSCPacket.h
@@ -1,7 +1,6 @@
 // Copyright Epic Games, Inc. All Rights Reserved.
 #pragma once
 
-#include "Interfaces/IPv4/IPv4Endpoint.h"
 #include "OSCTypes.h"
 #include "OSCStream.h"
 #include "OSCAddress.h"
@@ -10,7 +9,7 @@
 class FOSCStream;
 
 
-class IOSCPacket
+class OSC_API IOSCPacket
 {
 public:
 	IOSCPacket() = default;
@@ -28,12 +27,16 @@ public:
 	/** Returns true if packet is bundle */
 	virtual bool IsBundle() = 0;
 
-	/** Get endpoint responsible for creation/forwarding of packet */
-	virtual const FIPv4Endpoint& GetEndpoint() const;
+	/** Get endpoint IP address responsible for creation/forwarding of packet */
+	virtual const FString& GetIPAddress() const;
+
+	/** Get endpoint address responsible for creation/forwarding of packet */
+	virtual uint16 GetPort() const;
 
 	/** Create an OSC packet according to the input data. */
-	static TSharedPtr<IOSCPacket> CreatePacket(const uint8* InPacketType, const FIPv4Endpoint& InEndpoint);
+	static TSharedPtr<IOSCPacket> CreatePacket(const uint8* InPacketType, const FString& InAddress, uint16 InPort);
 
 protected:
-	FIPv4Endpoint Endpoint;
+	FString IPAddress;
+	uint16 Port;
 };
diff --git a/Engine/Plugins/Runtime/OSC/Source/OSC/Public/OSCStream.h b/Engine/Plugins/Runtime/OSC/Source/OSC/Public/OSCStream.h
--- a/Engine/Plugins/Runtime/OSC/Source/OSC/Public/OSCStream.h
+++ b/Engine/Plugins/Runtime/OSC/Source/OSC/Public/OSCStream.h
@@ -3,7 +3,7 @@
 
 #include "CoreMinimal.h"
 
-class FOSCStream
+class OSC_API FOSCStream
 {
 public:
 	FOSCStream();
diff --git a/Engine/Plugins/Runtime/OSC/Source/OSC/Public/OSCTypes.h b/Engine/Plugins/Runtime/OSC/Source/OSC/Public/OSCTypes.h
--- a/Engine/Plugins/Runtime/OSC/Source/OSC/Public/OSCTypes.h
+++ b/Engine/Plugins/Runtime/OSC/Source/OSC/Public/OSCTypes.h
@@ -25,7 +25,7 @@ enum class EOSCTypeTag : uint8
 };
 
 
-class FOSCType
+class OSC_API FOSCType
 {
 public:
 	explicit FOSCType(int32 Value)
