Fix for incorrect string conversions in FOSCStream::ReadChar and FOSCStream::WriteChar.

diff --git a/Engine/Plugins/Runtime/OSC/Source/OSC/Private/OSCStream.cpp b/Engine/Plugins/Runtime/OSC/Source/OSC/Private/OSCStream.cpp
--- a/Engine/Plugins/Runtime/OSC/Source/OSC/Private/OSCStream.cpp
+++ b/Engine/Plugins/Runtime/OSC/Source/OSC/Private/OSCStream.cpp
@@ -56,8 +56,8 @@ TCHAR FOSCStream::ReadChar()
 	uint8 Temp;
 	if (Read(&Temp, 1) > 0)
 	{
-		TCHAR* OutChar = ANSI_TO_TCHAR((ANSICHAR*)&Temp);
-		return *OutChar;
+		UE_CLOG(Temp > 0x7F, LogOSC, Warning, TEXT("Non-ANSI character '%u' written to OSCStream"), Temp);
+		return (TCHAR)Temp;
 	}
 
 	return '\0';
@@ -65,8 +65,10 @@ TCHAR FOSCStream::ReadChar()
 
 void FOSCStream::WriteChar(TCHAR Char)
 {
-	const uint8* ToWrite = (uint8*)TCHAR_TO_ANSI(&Char);
-	Write(ToWrite, 1);
+	const uint32 Temp = FChar::ToUnsigned(Char);
+	UE_CLOG(Temp > 0x7F, LogOSC, Warning, TEXT("Non-ANSI character '%u' written to OSCStream"), Temp);
+	const uint8 Temp8 = (uint8)Temp;
+	Write(&Temp8, 1);
 }
 
 FColor FOSCStream::ReadColor()
