OSC socket bind warning description

diff --git a/Engine/Plugins/Runtime/OSC/Source/OSC/Private/OSCServerProxy.cpp b/Engine/Plugins/Runtime/OSC/Source/OSC/Private/OSCServerProxy.cpp
--- a/Engine/Plugins/Runtime/OSC/Source/OSC/Private/OSCServerProxy.cpp
+++ b/Engine/Plugins/Runtime/OSC/Source/OSC/Private/OSCServerProxy.cpp
@@ -104,7 +104,8 @@ void FOSCServerProxy::Listen(const FString& InServerName)
 	}
 	else
 	{
-		UE_LOG(LogOSC, Warning, TEXT("OSCServer '%s' failed to bind to socket on %s:%d."), *InServerName, *ReceiveIPAddress.ToString(), Port);
+		// This is expected when the server isn't available, so it's not a Warning
+		UE_LOG(LogOSC, Display, TEXT("OSCServer '%s' failed to bind to socket on %s:%d. Check that the server is available on the specified address."), *InServerName, *ReceiveIPAddress.ToString(), Port);
 	}
 }
 
