OSC Plugin Optimization/Code Standard Refactor Part 1
- Move to using endpoints wherever possible in API vs string addr + int port
    - Allows for parse function to catch bad port IDs
    - Avoids string parsing when processing packets as an optimization
    - Added benefit of helping differentiate IP vs OSC addresses better in code and isolates use of IPv4 for potentially adding generic IP support & migrating socket implementation in the future
- Move to SharedRefs when handling packet passing to ensure failing fast on parse and avoiding a bunch of validity checks where unnecessary
- Optimize OSCTypes by using variant for all types to pack data better
- Clean-up namespaces/object names to clarify UObject/UStructs from proxy or non-uobject types & use more modern UE::OSC style (Helps as some types are now public to allow for non-GT/UObject client management and avoids future confusion)
- Remove CoreMinimal.h refs
- Add move symantics to OSCTypes

diff --git a/Engine/Plugins/Runtime/OSC/OSC.uplugin b/Engine/Plugins/Runtime/OSC/OSC.uplugin
--- a/Engine/Plugins/Runtime/OSC/OSC.uplugin
+++ b/Engine/Plugins/Runtime/OSC/OSC.uplugin
@@ -1,7 +1,7 @@
 {
 	"FileVersion" : 3,
 	"Version" : 1,
-	"VersionName" : "1.1",
+	"VersionName" : "2.0",
 	"FriendlyName" : "OSC (Open Sound Control)",
 	"Description" : "Implements the OSC 1.0 specification, allowing users to send and receive OSC messages and bundles between remote clients or applications.",
 	"Category" : "Input Devices",
diff --git a/Engine/Plugins/Runtime/OSC/Source/OSC/Private/OSCAddress.cpp b/Engine/Plugins/Runtime/OSC/Source/OSC/Private/OSCAddress.cpp
--- a/Engine/Plugins/Runtime/OSC/Source/OSC/Private/OSCAddress.cpp
+++ b/Engine/Plugins/Runtime/OSC/Source/OSC/Private/OSCAddress.cpp
@@ -5,11 +5,11 @@
 #include "OSCLog.h"
 
 
-namespace OSC
+namespace UE::OSC
 {
 	const FString BundleTag = TEXT("#bundle");
 	const FString PathSeparator = TEXT("/");
-} // namespace OSC
+} // namespace UE::OSC
 
 
 FOSCAddress::FOSCAddress()
@@ -23,7 +23,7 @@ FOSCAddress::FOSCAddress(const FString& InValue)
 	: bIsValidPattern(false)
 	, bIsValidPath(false)
 {
-	InValue.ParseIntoArray(Containers, *OSC::PathSeparator, true);
+	InValue.ParseIntoArray(Containers, *UE::OSC::PathSeparator, true);
 	if (Containers.Num() > 0)
 	{
 		Method = Containers.Pop();
@@ -62,56 +62,98 @@ bool FOSCAddress::IsValidPath() const
 	return bIsValidPath;
 }
 
-void FOSCAddress::PushContainer(const FString& Container)
+bool FOSCAddress::PushContainer(const FString& Container)
 {
-	if (Container.Contains(OSC::PathSeparator))
+	return PushContainers({ Container });
+}
+
+bool FOSCAddress::PushContainer(FString&& Container)
+{
+	return PushContainers({ MoveTemp(Container) });
+}
+
+bool FOSCAddress::PushContainers(const TArray<FString>& InContainers)
+{
+	TArray<FString> NewContainers = InContainers;
+	const bool bPushed = PushContainersInternal(MoveTemp(NewContainers));
+	if (bPushed)
 	{
-		UE_LOG(LogOSC, Warning, TEXT("Failed to push container on OSCAddress. "
-			"Cannot contain OSC path separator '%s'."), *OSC::PathSeparator);
-		return;
+		CacheAggregates();
 	}
 
-	Containers.Push(Container);
-	CacheAggregates();
+	return bPushed;
 }
 
-void FOSCAddress::PushContainers(const TArray<FString>& InContainers)
+bool FOSCAddress::PushContainers(TArray<FString>&& InContainers)
 {
+	const bool bPushed = PushContainersInternal(MoveTemp(InContainers));
+	if (bPushed)
+	{
+		CacheAggregates();
+	}
+
+	return bPushed;
+}
+
+bool FOSCAddress::PushContainersInternal(TArray<FString>&& InContainers)
+{
+	if (InContainers.IsEmpty())
+	{
+		return false;
+	}
+
 	for (const FString& Container : InContainers)
 	{
-		if (Container.Contains(OSC::PathSeparator))
+		if (Container.Contains(UE::OSC::PathSeparator))
 		{
 			UE_LOG(LogOSC, Warning, TEXT("Failed to push containers on OSCAddress. "
-				"Cannot contain OSC path separator '%s'."), *OSC::PathSeparator);
-			return;
+				"Cannot contain OSC path separator '%s'."), *UE::OSC::PathSeparator);
+			return false;
 		}
 	}
 
+	Containers.Append(MoveTemp(InContainers));
 	for (const FString& Container : InContainers)
 	{
 		Containers.Push(Container);
 	}
 
-	CacheAggregates();
+	return true;
 }
 
-FString FOSCAddress::PopContainer()
+FString FOSCAddress::PopContainer(bool* bOutDidPop)
 {
 	FString Popped;
 	if (Containers.Num() > 0)
 	{
 		Popped = Containers.Pop(EAllowShrinking::No);
 		Hash = GetTypeHash(GetFullPath());
+
+		if (bOutDidPop)
+		{
+			*bOutDidPop = true;
+		}
+	}
+	else
+	{
+		if (bOutDidPop)
+		{
+			*bOutDidPop = false;
+		}
 	}
 
 	return Popped;
 }
 
-TArray<FString> FOSCAddress::PopContainers(int32 InNumContainers)
+TArray<FString> FOSCAddress::PopContainers(int32 InNumContainers, bool* bOutDidPop)
 {
 	TArray<FString> Popped;
 	if (InNumContainers <= 0 || Containers.Num() == 0)
 	{
+		if (bOutDidPop)
+		{
+			*bOutDidPop = false;
+		}
 		return Popped;
 	}
 
@@ -127,18 +169,25 @@ TArray<FString> FOSCAddress::PopContainers(int32 InNumContainers)
 	}
 
 	Hash = GetTypeHash(GetFullPath());
+	if (bOutDidPop)
+	{
+		*bOutDidPop = true;
+	}
 	return Popped;
 }
 
-void FOSCAddress::RemoveContainers(int32 InIndex, int32 InCount)
+bool FOSCAddress::RemoveContainers(int32 InIndex, int32 InCount)
 {
 	if (InIndex >= 0 && InCount > 0)
 	{
 		if (InIndex + InCount < Containers.Num())
 		{
 			Containers.RemoveAt(InIndex, InCount);
+			return true;
 		}
 	}
+
+	return false;
 }
 
 void FOSCAddress::ClearContainers()
@@ -152,29 +201,109 @@ const FString& FOSCAddress::GetMethod() const
 	return Method;
 }
 
-void FOSCAddress::SetMethod(const FString& InMethod)
+int32 FOSCAddress::GetNumericPrefix(bool* bIsNumeric) const
+{
+	int32 Value = 0;
+	int32 Dec = 1;
+	const TArray<TCHAR>& CharArray = Method.GetCharArray();
+	for (
+		int32 Index = 0;
+		Index < CharArray.Num() && FChar::IsDigit(CharArray[Index]);
+		++Index)
+	{
+		Value = (Value * Dec) + FChar::ConvertCharDigitToInt(CharArray[Index]);
+		Dec *= 10;
+	}
+
+	if (bIsNumeric)
+	{
+		*bIsNumeric = Dec == 1;
+	}
+	return Value;
+}
+
+int32 FOSCAddress::GetNumericSuffix(bool* bIsNumeric) const
+{
+	int32 Value = 0;
+	int32 Dec = 1;
+	const TArray<TCHAR>& CharArray = Method.GetCharArray();
+	for (
+		int32 Index = CharArray.Num() - 2 /* less term char */;
+		Index >= 0 && FChar::IsDigit(CharArray[Index]);
+		--Index)
+	{
+		Value += FChar::ConvertCharDigitToInt(CharArray[Index]) * Dec;
+		Dec *= 10;
+	}
+
+	if (bIsNumeric)
+	{
+		*bIsNumeric = Dec == 1;
+	}
+	return Value;
+}
+
+bool FOSCAddress::Set(const TArray<FString>& InContainers, const FString& InMethod)
+{
+	TArray<FString> NewContainers = InContainers;
+	FString NewMethod = InMethod;
+	return Set(MoveTemp(NewContainers), MoveTemp(NewMethod));
+}
+
+bool FOSCAddress::Set(TArray<FString>&& InContainers, FString&& InMethod)
+{
+	Containers = { };
+	Method = { };
+
+	const bool bPushContainers = PushContainersInternal(MoveTemp(InContainers));
+	if (bPushContainers)
+	{
+		const bool bSetMethod = SetMethod(MoveTemp(InMethod)); // Calls aggregate internally, so no need to call again on success
+		if (bSetMethod)
+		{
+			return true;
+		}
+		else
+		{
+			Containers = { };
+		}
+	}
+
+	CacheAggregates();
+	return false;
+}
+
+bool FOSCAddress::SetMethod(const FString& InMethod)
+{
+	FString NewMethod = InMethod;
+	return SetMethod(MoveTemp(NewMethod));
+}
+
+bool FOSCAddress::SetMethod(FString&& InMethod)
 {
 	if (InMethod.IsEmpty())
 	{
 		UE_LOG(LogOSC, Warning, TEXT("Failed to set OSCAddress method. "
 			"'InMethod' cannot be empty string."));
-		return;
+		return false;
 	}
 
-	if (InMethod.Contains(OSC::PathSeparator))
+	if (InMethod.Contains(UE::OSC::PathSeparator))
 	{
 		UE_LOG(LogOSC, Warning, TEXT("Failed to set OSCAddress method. "
-			"Cannot contain OSC path separator '%s'."), *OSC::PathSeparator);
-		return;
+			"Cannot contain OSC path separator '%s'."), *UE::OSC::PathSeparator);
+		return false;
 	}
 
-	Method = InMethod;
+	Method = MoveTemp(InMethod);
+
 	CacheAggregates();
+	return true;
 }
 
 FString FOSCAddress::GetContainerPath() const
 {
-	return OSC::PathSeparator + FString::Join(Containers, *OSC::PathSeparator);
+	return UE::OSC::PathSeparator + FString::Join(Containers, *UE::OSC::PathSeparator);
 }
 
 FString FOSCAddress::GetContainer(int32 Index) const
@@ -196,8 +325,8 @@ FString FOSCAddress::GetFullPath() const
 {
 	if (Containers.Num() == 0)
 	{
-		return OSC::PathSeparator + Method;
+		return UE::OSC::PathSeparator + Method;
 	}
 
-	return GetContainerPath() + OSC::PathSeparator + Method;
+	return GetContainerPath() + UE::OSC::PathSeparator + Method;
 }
diff --git a/Engine/Plugins/Runtime/OSC/Source/OSC/Private/OSCBundle.cpp b/Engine/Plugins/Runtime/OSC/Source/OSC/Private/OSCBundle.cpp
--- a/Engine/Plugins/Runtime/OSC/Source/OSC/Private/OSCBundle.cpp
+++ b/Engine/Plugins/Runtime/OSC/Source/OSC/Private/OSCBundle.cpp
@@ -5,26 +5,42 @@
 
 
 FOSCBundle::FOSCBundle()
-	: Packet(MakeShareable(new FOSCBundlePacket()))
+	: Packet(MakeShareable(new UE::OSC::FBundlePacket()))
 {
 }
 
-FOSCBundle::FOSCBundle(const TSharedPtr<IOSCPacket>& InPacket)
+FOSCBundle::FOSCBundle(const TSharedPtr<UE::OSC::IPacket>& InPacket)
+	: Packet(InPacket.Get())
+{
+}
+
+FOSCBundle::FOSCBundle(const TSharedRef<UE::OSC::IPacket>& InPacket)
 	: Packet(InPacket)
 {
 }
 
-FOSCBundle::~FOSCBundle()
+void FOSCBundle::SetPacket(TSharedPtr<UE::OSC::IPacket>& InPacket)
 {
-	Packet.Reset();
+	using namespace UE::OSC;
+
+	check(InPacket->IsBundle());
+	Packet = TSharedRef<IPacket>(InPacket.Get());
 }
 
-void FOSCBundle::SetPacket(const TSharedPtr<IOSCPacket>& InPacket)
+void FOSCBundle::SetPacket(const TSharedRef<UE::OSC::IPacket>& InPacket)
 {
+	check(InPacket->IsBundle());
 	Packet = InPacket;
 }
 
-const TSharedPtr<IOSCPacket>& FOSCBundle::GetPacket() const
+const TSharedPtr<UE::OSC::IPacket>& FOSCBundle::GetPacket() const
+{
+	static TSharedPtr<UE::OSC::IPacket> PacketPtr;
+	PacketPtr = TSharedPtr<UE::OSC::IPacket>(&Packet.Get());
+	return PacketPtr;
+}
+
+const TSharedRef<UE::OSC::IPacket>& FOSCBundle::GetPacketRef() const
 {
 	return Packet;
 }
diff --git a/Engine/Plugins/Runtime/OSC/Source/OSC/Private/OSCBundlePacket.cpp b/Engine/Plugins/Runtime/OSC/Source/OSC/Private/OSCBundlePacket.cpp
--- a/Engine/Plugins/Runtime/OSC/Source/OSC/Private/OSCBundlePacket.cpp
+++ b/Engine/Plugins/Runtime/OSC/Source/OSC/Private/OSCBundlePacket.cpp
@@ -6,94 +6,98 @@
 #include "OSCStream.h"
 
 
-FOSCBundlePacket::FOSCBundlePacket()
-	: IOSCPacket()
-	, TimeTag(0)
+namespace UE::OSC
 {
-}
-
-FOSCBundlePacket::~FOSCBundlePacket()
-{
-}
-
-void FOSCBundlePacket::SetTimeTag(uint64 NewTimeTag)
-{
-	TimeTag = FOSCType(NewTimeTag);
-}
+	FBundlePacket::FBundlePacket()
+		: TimeTag(0)
+	{
+	}
 
-uint64 FOSCBundlePacket::GetTimeTag() const
-{
-	return TimeTag.GetTimeTag();
-}
+	FBundlePacket::FBundlePacket(const FIPv4Endpoint& InEndpoint)
+		: FPacketBase(InEndpoint)
+		, TimeTag(0)
+	{
+	}
 
-void FOSCBundlePacket::WriteData(FOSCStream& Stream)
-{
-	// Write bundle & time tag
-	Stream.WriteString(OSC::BundleTag);
-	Stream.WriteUInt64(GetTimeTag());
+	void FBundlePacket::SetTimeTag(uint64 NewTimeTag)
+	{
+		TimeTag = FOSCData(NewTimeTag);
+	}
 
-	for (TSharedPtr<IOSCPacket>& Packet : Packets)
+	uint64 FBundlePacket::GetTimeTag() const
 	{
-		int32 StreamPos = Stream.GetPosition();
-		Stream.WriteInt32(0);
+		return TimeTag.GetTimeTag();
+	}
 
-		int32 InitPos = Stream.GetPosition();
-		Packet->WriteData(Stream);
-		int32 NewPos = Stream.GetPosition();
+	void FBundlePacket::WriteData(FStream& Stream)
+	{
+		// Write bundle & time tag
+		Stream.WriteString(OSC::BundleTag);
+		Stream.WriteUInt64(GetTimeTag());
 
-		Stream.SetPosition(StreamPos);
-		Stream.WriteInt32(NewPos - InitPos);
-		Stream.SetPosition(NewPos);
-	}
-}
+		for (const TSharedRef<UE::OSC::IPacket>& Packet : Packets)
+		{
+			int32 StreamPos = Stream.GetPosition();
+			Stream.WriteInt32(0);
 
-void FOSCBundlePacket::ReadData(FOSCStream& Stream)
-{
-	Packets.Reset();
+			int32 InitPos = Stream.GetPosition();
+			Packet->WriteData(Stream);
+			int32 NewPos = Stream.GetPosition();
 
-	FString BundleTag = Stream.ReadString();
-	if (BundleTag != OSC::BundleTag)
-	{
-		UE_LOG(LogOSC, Warning, TEXT("Failed to parse OSCBundle of invalid format. #bundle identifier not first item in packet."));
-		return;
+			Stream.SetPosition(StreamPos);
+			Stream.WriteInt32(NewPos - InitPos);
+			Stream.SetPosition(NewPos);
+		}
 	}
 
-	TimeTag = FOSCType(Stream.ReadUInt64());
-
-	while (!Stream.HasReachedEnd())
+	void FBundlePacket::ReadData(FStream& Stream)
 	{
-		int32 PacketLength = Stream.ReadInt32();
+		Packets.Reset();
 
-		int32 StartPos = Stream.GetPosition();
-		TSharedPtr<IOSCPacket> Packet = IOSCPacket::CreatePacket(Stream.GetData() + Stream.GetPosition(), IPAddress, Port);
-		if (!Packet.IsValid())
+		const FString ThisBundleTag = Stream.ReadString();
+		if (ThisBundleTag != UE::OSC::BundleTag)
 		{
-			break;
+			UE_LOG(LogOSC, Warning, TEXT("Failed to parse OSCBundle of invalid format. #bundle identifier not first item in packet."));
+			return;
 		}
 
-		Packet->ReadData(Stream);
-		Packets.Add(Packet);
-		int32 EndPos = Stream.GetPosition();
+		TimeTag = FOSCData(Stream.ReadUInt64());
 
-		if (EndPos - StartPos != PacketLength)
+		while (!Stream.HasReachedEnd())
 		{
-			UE_LOG(LogOSC, Warning, TEXT("Failed to parse OSCBundle of invalid format. Element size mismatch."));
-			break;
+			int32 PacketLength = Stream.ReadInt32();
+
+			int32 StartPos = Stream.GetPosition();
+			TSharedPtr<UE::OSC::IPacket> Packet = UE::OSC::IPacket::CreatePacket(Stream.GetData() + Stream.GetPosition(), IPEndpoint);
+			if (!Packet.IsValid())
+			{
+				break;
+			}
+
+			Packet->ReadData(Stream);
+			Packets.Add(Packet->AsShared());
+			int32 EndPos = Stream.GetPosition();
+
+			if (EndPos - StartPos != PacketLength)
+			{
+				UE_LOG(LogOSC, Warning, TEXT("Failed to parse OSCBundle of invalid format. Element size mismatch."));
+				break;
+			}
 		}
 	}
-}
 
-FOSCBundlePacket::FPacketBundle& FOSCBundlePacket::GetPackets()
-{
-	return Packets;
-}
+	TArray<TSharedRef<UE::OSC::IPacket>>& FBundlePacket::GetPackets()
+	{
+		return Packets;
+	}
 
-bool FOSCBundlePacket::IsBundle()
-{
-	return true;
-}
+	bool FBundlePacket::IsBundle()
+	{
+		return true;
+	}
 
-bool FOSCBundlePacket::IsMessage()
-{
-	return false;
-}
+	bool FBundlePacket::IsMessage()
+	{
+		return false;
+	}
+} // namespace UE::OSC
diff --git a/Engine/Plugins/Runtime/OSC/Source/OSC/Private/OSCBundlePacket.h b/Engine/Plugins/Runtime/OSC/Source/OSC/Private/OSCBundlePacket.h
--- a/Engine/Plugins/Runtime/OSC/Source/OSC/Private/OSCBundlePacket.h
+++ b/Engine/Plugins/Runtime/OSC/Source/OSC/Private/OSCBundlePacket.h
@@ -1,41 +1,45 @@
 // Copyright Epic Games, Inc. All Rights Reserved.
 #pragma once
 
+#include "OSCMessagePacket.h"
 #include "OSCPacket.h"
 #include "OSCTypes.h"
 
 
-class OSC_API FOSCBundlePacket : public IOSCPacket
+namespace UE::OSC
 {
-public:
-	FOSCBundlePacket();
-	virtual ~FOSCBundlePacket();
+	class FBundlePacket : public FPacketBase
+	{
+	public:
+		FBundlePacket();
+		FBundlePacket(const FIPv4Endpoint& InEndpoint);
 
-	using FPacketBundle = TArray<TSharedPtr<IOSCPacket>>;
+		virtual ~FBundlePacket() = default;
 
-	/** Set the bundle time tag. */
-	void SetTimeTag(uint64 NewTimeTag);
+		/** Set the bundle time tag. */
+		void SetTimeTag(uint64 NewTimeTag);
 
-	/** Get the bundle time tag. */
-	uint64 GetTimeTag() const;
+		/** Get the bundle time tag. */
+		uint64 GetTimeTag() const;
 
-	/** Get OSC packet by index. */
-	FPacketBundle& GetPackets();
+		/** Get OSC packet bundle. */
+		TArray<TSharedRef<UE::OSC::IPacket>>& GetPackets();
 
-	virtual bool IsBundle() override;
-	virtual bool IsMessage() override;
+		virtual bool IsBundle() override;
+		virtual bool IsMessage() override;
 
-	/** Writes bundle data into the OSC stream. */
-	virtual void WriteData(FOSCStream& Stream) override;
+		/** Writes bundle data into the OSC stream. */
+		virtual void WriteData(FStream& Stream) override;
 
-	/** Reads bundle data from provided OSC stream,
-	  * adding packet data to internal packet bundle. */
-	virtual void ReadData(FOSCStream& Stream) override;
+		/** Reads bundle data from provided OSC stream,
+		  * adding packet data to internal packet bundle. */
+		virtual void ReadData(FStream& Stream) override;
 
-private:
-	/** Bundle of OSC packets. */
-	FPacketBundle Packets;
+	private:
+		/** Array of OSC packets. */
+		TArray<TSharedRef<UE::OSC::IPacket>> Packets;
 
-	/** Bundle time tag. */
-	FOSCType TimeTag;
-};
+		/** Bundle time tag. */
+		FOSCData TimeTag;
+	};
+} // namespace UE::OSC
diff --git a/Engine/Plugins/Runtime/OSC/Source/OSC/Private/OSCClient.cpp b/Engine/Plugins/Runtime/OSC/Source/OSC/Private/OSCClient.cpp
--- a/Engine/Plugins/Runtime/OSC/Source/OSC/Private/OSCClient.cpp
+++ b/Engine/Plugins/Runtime/OSC/Source/OSC/Private/OSCClient.cpp
@@ -1,10 +1,17 @@
 // Copyright Epic Games, Inc. All Rights Reserved.
-#include "OSCClient.h"
-
 
+#include "OSCClient.h"
 #include "OSCClientProxy.h"
 
 
+namespace UE::OSC
+{
+	TUniquePtr<IClientProxy> IClientProxy::Create(const FString& ClientName)
+	{
+		return MakeUnique<UE::OSC::FClientProxy>(ClientName);
+	}
+} // namespace UE::OSC
+
 UOSCClient::UOSCClient(const FObjectInitializer& ObjectInitializer)
 	: Super(ObjectInitializer)
 	, ClientProxy(nullptr)
@@ -14,7 +21,7 @@ UOSCClient::UOSCClient(const FObjectInitializer& ObjectInitializer)
 void UOSCClient::Connect()
 {
 	check(!ClientProxy.IsValid());
-	ClientProxy.Reset(new FOSCClientProxy(GetName()));
+	ClientProxy = UE::OSC::IClientProxy::Create(GetName());
 }
 
 bool UOSCClient::IsActive() const
@@ -25,13 +32,26 @@ bool UOSCClient::IsActive() const
 void UOSCClient::GetSendIPAddress(FString& InIPAddress, int32& Port)
 {
 	check(ClientProxy.IsValid());
-	ClientProxy->GetSendIPAddress(InIPAddress, Port);
+
+	const FIPv4Endpoint Endpoint = ClientProxy->GetSendIPEndpoint();
+	InIPAddress = Endpoint.Address.ToString();
+	Port = Endpoint.Port;
 }
 
 bool UOSCClient::SetSendIPAddress(const FString& InIPAddress, const int32 Port)
 {
 	check(ClientProxy.IsValid());
-	return ClientProxy->SetSendIPAddress(InIPAddress, Port);
+
+	FIPv4Endpoint Endpoint;
+	if (!FIPv4Address::Parse(InIPAddress, Endpoint.Address))
+	{
+		UE_LOG(LogOSC, Warning, TEXT("Failed to set send IP Address: Could not parse IPAddress from string '%s'."), *InIPAddress);
+		return false;
+	}
+	Endpoint.Port = Port;
+
+	ClientProxy->SetSendIPEndpoint(Endpoint);
+	return true;
 }
 
 void UOSCClient::Stop()
diff --git a/Engine/Plugins/Runtime/OSC/Source/OSC/Private/OSCClientProxy.cpp b/Engine/Plugins/Runtime/OSC/Source/OSC/Private/OSCClientProxy.cpp
--- a/Engine/Plugins/Runtime/OSC/Source/OSC/Private/OSCClientProxy.cpp
+++ b/Engine/Plugins/Runtime/OSC/Source/OSC/Private/OSCClientProxy.cpp
@@ -9,123 +9,113 @@
 #include "OSCMessagePacket.h"
 #include "OSCStream.h"
 
-namespace
+namespace UE::OSC
 {
 	static const int32 OUTPUT_BUFFER_SIZE = 1024;
-} // namespace <>
 
-
-FOSCClientProxy::FOSCClientProxy(const FString& InClientName)
-	: Socket(FUdpSocketBuilder(*InClientName).Build())
-{
-	ClientName = InClientName;
-}
-
-FOSCClientProxy::~FOSCClientProxy()
-{
-	Stop();
-}
-
-void FOSCClientProxy::GetSendIPAddress(FString& InIPAddress, int32& Port) const
-{
-	const bool bAppendPort = false;
-	InIPAddress = IPAddress->ToString(bAppendPort);
-	Port = IPAddress->GetPort();
-}
-
-bool FOSCClientProxy::SetSendIPAddress(const FString& InIPAddress, const int32 Port)
-{
-	IPAddress = ISocketSubsystem::Get(PLATFORM_SOCKETSUBSYSTEM)->CreateInternetAddr();
-
-	bool bIsValidAddress = true;
-	IPAddress->SetIp(*InIPAddress, bIsValidAddress);
-	IPAddress->SetPort(Port);
-
-	if (bIsValidAddress)
+	FClientProxy::FClientProxy(const FString& InClientName)
+		: Socket(FUdpSocketBuilder(InClientName).Build())
 	{
-		UE_LOG(LogOSC, Verbose, TEXT("OSCClient '%s' SetSendIpAddress: %s:%d"), *ClientName, *InIPAddress, Port);
 	}
-	else
+
+	FClientProxy::~FClientProxy()
 	{
-		UE_LOG(LogOSC, Warning, TEXT("OSCClient '%s' SetSendIpAddress Failed for input: %s:%d"), *ClientName, *InIPAddress, Port);
+		if (Socket)
+		{
+			ISocketSubsystem::Get(PLATFORM_SOCKETSUBSYSTEM)->DestroySocket(Socket);
+			Socket = nullptr;
+		}
 	}
 
-	return bIsValidAddress;
-}
-
-bool FOSCClientProxy::IsActive() const
-{
-	return Socket && Socket->GetConnectionState() == ESocketConnectionState::SCS_Connected;
-}
+	const FIPv4Endpoint& FClientProxy::GetSendIPEndpoint() const
+	{
+		return IPEndpoint;
+	}
 
-void FOSCClientProxy::SendPacket(IOSCPacket& Packet)
-{
-	if (!Socket)
+	void FClientProxy::SetSendIPEndpoint(const FIPv4Endpoint& InEndpoint)
 	{
-		UE_LOG(LogOSC, Error, TEXT("Socket has been closed. OSCClient '%s' failed to send msg"), *ClientName);
-		return;
+		IPEndpoint = InEndpoint;
 	}
 
-	const FOSCAddress* OSCAddress = nullptr;
-	if (Packet.IsMessage())
+	bool FClientProxy::IsActive() const
 	{
-		const FOSCMessagePacket& MessagePacket = static_cast<FOSCMessagePacket&>(Packet);
-		OSCAddress = &MessagePacket.GetAddress();
-		if (!OSCAddress->IsValidPath())
-		{
-			UE_LOG(LogOSC, Warning, TEXT("Failed to write packet data. Invalid OSCAddress '%s'"), *OSCAddress->GetFullPath());
-			return;
-		}
+		return Socket && Socket->GetConnectionState() == ESocketConnectionState::SCS_Connected;
 	}
 
-	if (IPAddress)
+	void FClientProxy::SendPacket(UE::OSC::IPacket& Packet)
 	{
-		FOSCStream Stream = FOSCStream();
-		Packet.WriteData(Stream);
-		const uint8* DataPtr = Stream.GetData();
+		if (!Socket)
+		{
+#if !UE_BUILD_SHIPPING
+			UE_LOG(LogOSC, Error, TEXT("OSCClient stopped (socket '%s') has been stopped. Failed to send msg"), *DestroyedSocketDesc);
+#endif // !UE_BUILD_SHIPPING
 
-		int32 BytesSent = 0;
+			return;
+		}
 
-		const int32 AttemptedLength = Stream.GetPosition();
-		int32 Length = AttemptedLength;
-		while (Length > 0)
+		const FOSCAddress* OSCAddress = nullptr;
+		if (Packet.IsMessage())
 		{
-			const bool bSuccess = Socket->SendTo(DataPtr, Length, BytesSent, *IPAddress);
-			if (!bSuccess || BytesSent <= 0)
+			const FMessagePacket& MessagePacket = static_cast<FMessagePacket&>(Packet);
+			OSCAddress = &MessagePacket.GetAddress();
+			if (!OSCAddress->IsValidPath())
 			{
-				UE_LOG(LogOSC, Verbose, TEXT("OSC Packet failed: Client '%s', OSC Identifier '%s', Send IPAddress %s, Attempted Bytes = %d"),
-					*ClientName, OSCAddress ? *OSCAddress->GetFullPath() : *OSC::BundleTag, *IPAddress->ToString(true /*bAppendPort*/), AttemptedLength);
+				UE_LOG(LogOSC, Warning, TEXT("Failed to write packet data. Invalid OSCAddress '%s'"), *OSCAddress->GetFullPath());
 				return;
 			}
-
-			Length -= BytesSent;
-			DataPtr += BytesSent;
 		}
 
-		UE_LOG(LogOSC, Verbose, TEXT("OSC Packet sent: Client '%s', OSC Identifier '%s', Send IPAddress %s, Bytes Sent = %d"),
-			*ClientName, OSCAddress ? *OSCAddress->GetFullPath() : *OSC::BundleTag, *IPAddress->ToString(true /*bAppendPort*/), AttemptedLength);
+		if (IPEndpoint != FIPv4Endpoint::Any)
+		{
+			TSharedRef<FInternetAddr> InternetAddr = IPEndpoint.ToInternetAddr();
+			FStream Stream = FStream();
+			Packet.WriteData(Stream);
+			const uint8* DataPtr = Stream.GetData();
+
+			int32 BytesSent = 0;
+
+			const int32 AttemptedLength = Stream.GetPosition();
+			int32 Length = AttemptedLength;
+			while (Length > 0)
+			{
+				const bool bSuccess = Socket->SendTo(DataPtr, Length, BytesSent, *InternetAddr);
+				if (!bSuccess || BytesSent <= 0)
+				{
+					UE_LOG(LogOSC, Verbose, TEXT("OSC Packet failed: Client '%s', OSC Address '%s', Send IP Endpoint %s, Attempted Bytes = %d"),
+						*Socket->GetDescription(), OSCAddress ? *OSCAddress->GetFullPath() : *UE::OSC::BundleTag, *IPEndpoint.ToString(), AttemptedLength);
+					return;
+				}
+
+				Length -= BytesSent;
+				DataPtr += BytesSent;
+			}
+
+			UE_LOG(LogOSC, Verbose, TEXT("OSC Packet sent: Client '%s', OSC Address '%s', Send IP Endpoint %s, Bytes Sent = %d"),
+				*Socket->GetDescription(), OSCAddress ? *OSCAddress->GetFullPath() : *UE::OSC::BundleTag, *IPEndpoint.ToString(), AttemptedLength);
+		}
 	}
-}
 
-void FOSCClientProxy::SendMessage(FOSCMessage& Message)
-{
-	const TSharedPtr<IOSCPacket>& Packet = Message.GetPacket();
-	check(Packet.IsValid());
-	SendPacket(*Packet.Get());
-}
+	void FClientProxy::SendMessage(FOSCMessage& Message)
+	{
+		const TSharedRef<UE::OSC::IPacket>& Packet = Message.GetPacketRef();
+		SendPacket(Packet.Get());
+	}
 
-void FOSCClientProxy::SendBundle(FOSCBundle& Bundle)
-{
-	const TSharedPtr<IOSCPacket>& Packet = Bundle.GetPacket();
-	check(Packet.IsValid());
-	SendPacket(*Packet.Get());
-}
+	void FClientProxy::SendBundle(FOSCBundle& Bundle)
+	{
+		const TSharedRef<UE::OSC::IPacket>& Packet = Bundle.GetPacketRef();
+		SendPacket(Packet.Get());
+	}
 
-void FOSCClientProxy::Stop()
-{
-	if (Socket)
+	void FClientProxy::Stop()
 	{
-		ISocketSubsystem::Get(PLATFORM_SOCKETSUBSYSTEM)->DestroySocket(Socket);
-		Socket = nullptr;
+		if (Socket)
+		{
+#if !UE_BUILD_SHIPPING
+			DestroyedSocketDesc = Socket->GetDescription();
+#endif // !UE_BUILD_SHIPPING
+			ISocketSubsystem::Get(PLATFORM_SOCKETSUBSYSTEM)->DestroySocket(Socket);
+			Socket = nullptr;
+		}
 	}
-}
+} // namespace UE::OSC
diff --git a/Engine/Plugins/Runtime/OSC/Source/OSC/Private/OSCClientProxy.h b/Engine/Plugins/Runtime/OSC/Source/OSC/Private/OSCClientProxy.h
--- a/Engine/Plugins/Runtime/OSC/Source/OSC/Private/OSCClientProxy.h
+++ b/Engine/Plugins/Runtime/OSC/Source/OSC/Private/OSCClientProxy.h
@@ -1,40 +1,52 @@
 // Copyright Epic Games, Inc. All Rights Reserved.
 #pragma once
 
-
-
 #include "OSCClient.h"
 
+
+// Forward Declarations
 class FInternetAddr;
 class FSocket;
-class IOSCPacket;
 
 
-class OSC_API FOSCClientProxy : public IOSCClientProxy
+namespace UE::OSC
 {
-public:
-	FOSCClientProxy(const FString& InClientName);
-	virtual ~FOSCClientProxy();
+	// Forward Declarations
+	class IPacket;
+
+	class FClientProxy : public IClientProxy
+	{
+	public:
+		FClientProxy(const FString& InClientName);
+		virtual ~FClientProxy();
+
+		UE_DEPRECATED(5.5, "Use GetSendIPEndpoint instead")
+		virtual void GetSendIPAddress(FString& InIPAddress, int32& Port) const { }
+
+		UE_DEPRECATED(5.5, "Use SetSendIPEndpoint instead")
+		virtual bool SetSendIPAddress(const FString& InIPAddress, const int32 Port) { return false; }
 
-	void GetSendIPAddress(FString& InIPAddress, int32& Port) const override;
-	bool SetSendIPAddress(const FString& InIPAddress, const int32 Port) override;
+		virtual const FIPv4Endpoint& GetSendIPEndpoint() const override;
+		virtual void SetSendIPEndpoint(const FIPv4Endpoint& InEndpoint) override;
 
-	bool IsActive() const override;
+		bool IsActive() const override;
 
-	void SendMessage(FOSCMessage& Message) override;
-	void SendBundle(FOSCBundle& Bundle) override;
+		void SendMessage(FOSCMessage& Message) override;
+		void SendBundle(FOSCBundle& Bundle) override;
 
-	void Stop() override;
+		void Stop() override;
 
-private:
-	void SendPacket(IOSCPacket& Packet);
+		void SendPacket(UE::OSC::IPacket& Packet);
 
-	/** Socket used to send the OSC packets. */
-	FSocket* Socket;
+	private:
+		/** Socket used to send the OSC packets. */
+		FSocket* Socket = nullptr;
 
-	/** IP Address used by socket. */
-	TSharedPtr<FInternetAddr> IPAddress;
+		/** IP Address used by socket. */
+		FIPv4Endpoint IPEndpoint;
 
-	/** Name of client */
-	FString ClientName;
-};
+#if !UE_BUILD_SHIPPING
+		FString DestroyedSocketDesc;
+#endif // !UE_BUILD_SHIPPING
+	};
+} // namespace UE::OSC
diff --git a/Engine/Plugins/Runtime/OSC/Source/OSC/Private/OSCManager.cpp b/Engine/Plugins/Runtime/OSC/Source/OSC/Private/OSCManager.cpp
--- a/Engine/Plugins/Runtime/OSC/Source/OSC/Private/OSCManager.cpp
+++ b/Engine/Plugins/Runtime/OSC/Source/OSC/Private/OSCManager.cpp
@@ -18,334 +18,360 @@
 #include "UObject/UObjectIterator.h"
 
 
-#define OSC_LOG_INVALID_TYPE_AT_INDEX(TypeStr, Index, Msg) UE_LOG(LogOSC, Warning, TEXT("OSC Message Parse Failed: OSCType not %s: index '%i', OSCAddress '%s'"), TypeStr, Index, *Msg.GetAddress().GetFullPath())
-
-namespace OSC
+namespace UE::OSC
 {
-	static const int32 DefaultClientPort = 8094;
-	static const int32 DefaultServerPort = 8095;
-
-	// Returns true if provided address was null and was able to
-	// override with local host address, false if not.
-	bool GetLocalHostAddress(FString& InAddress)
+	namespace ManagerPrivate
 	{
-		if (!InAddress.IsEmpty() && InAddress != TEXT("0"))
+		static const int32 DefaultClientPort = 8094;
+		static const int32 DefaultServerPort = 8095;
+
+		// Returns true if provided address was null and was able to
+		// override with local host address, false if not.
+		bool GetLocalHostAddress(FString& InAddress)
 		{
+			if (!InAddress.IsEmpty() && InAddress != TEXT("0"))
+			{
+				return false;
+			}
+
+			bool bCanBind = false;
+			bool bAppendPort = false;
+			if (ISocketSubsystem* SocketSys = ISocketSubsystem::Get(PLATFORM_SOCKETSUBSYSTEM))
+			{
+				const TSharedPtr<FInternetAddr> Addr = SocketSys->GetLocalHostAddr(*GLog, bCanBind);
+				if (Addr.IsValid())
+				{
+					InAddress = Addr->ToString(bAppendPort);
+					return true;
+				}
+			}
+
 			return false;
 		}
 
-		bool bCanBind = false;
-		bool bAppendPort = false;
-		if (ISocketSubsystem* SocketSys = ISocketSubsystem::Get(PLATFORM_SOCKETSUBSYSTEM))
+		const FOSCData* GetDataAtIndex(const FOSCMessage& InMessage, const int32 InIndex)
 		{
-			const TSharedPtr<FInternetAddr> Addr = SocketSys->GetLocalHostAddr(*GLog, bCanBind);
-			if (Addr.IsValid())
+			const TSharedRef<FMessagePacket>& Packet = StaticCastSharedRef<FMessagePacket>(InMessage.GetPacketRef());
+			const TArray<FOSCData>& Args = Packet->GetArguments();
+			if (Args.IsValidIndex(InIndex))
 			{
-				InAddress = Addr->ToString(bAppendPort);
-				return true;
+				return &Args[InIndex];
 			}
-		}
 
-		return false;
-	}
+			UE_LOG(LogOSC, Warning, TEXT("Index '%d' out-of-bounds.  Message argument size = '%d'"), InIndex, Args.Num());
+			return nullptr;
+		}
 
-	const FOSCType* GetOSCTypeAtIndex(const FOSCMessage& InMessage, const int32 InIndex)
-	{
-		const TSharedPtr<FOSCMessagePacket>& Packet = StaticCastSharedPtr<FOSCMessagePacket>(InMessage.GetPacket());
-		if (Packet.IsValid())
+		void IterateMessageArgs(const FOSCMessage& InMessage, TFunctionRef<void(const FOSCData&)> InFunc)
 		{
-			TArray<FOSCType>& Args = Packet->GetArguments();
-			if (InIndex >= Args.Num())
+			const TSharedRef<FMessagePacket>& MessagePacket = StaticCastSharedRef<FMessagePacket>(InMessage.GetPacketRef());
+			const TArray<FOSCData>& Args = MessagePacket->GetArguments();
+			for (int32 i = 0; i < Args.Num(); i++)
 			{
-				UE_LOG(LogOSC, Warning, TEXT("Index '%d' out-of-bounds.  Message argument size = '%d'"), InIndex, Args.Num());
-				return nullptr;
+				const FOSCData& OSCData = Args[i];
+				InFunc(OSCData);
 			}
-
-			return &Args[InIndex];
 		}
 
-		return nullptr;
-	}
-} // namespace OSC
-
-
-static FAutoConsoleCommand GOSCPrintServers(
-	TEXT("osc.servers"),
-	TEXT("Prints diagnostic information pertaining to the currently initialized OSC servers objects to the output log."),
-	FConsoleCommandDelegate::CreateStatic(
-		[]()
+		void LogInvalidTypeAtIndex(EDataType DataType, int32 Index, const FOSCMessage& Msg)
 		{
-			FString LocalAddr;
-			OSC::GetLocalHostAddress(LocalAddr);
-			UE_LOG(LogOSC, Display, TEXT("Local IP: %s"), *LocalAddr);
+			UE_LOG(LogOSC, Warning, TEXT("OSC Message Parse Failed: OSCData not %s: index '%i', OSCAddress '%s'"),
+				LexToString(DataType),
+				Index,
+				*Msg.GetAddress().GetFullPath());
+		}
 
-			UE_LOG(LogOSC, Display, TEXT("OSC Servers:"));
-			for (TObjectIterator<UOSCServer> Itr; Itr; ++Itr)
-			{
-				if (UOSCServer* Server = *Itr)
+		static FAutoConsoleCommand GOSCPrintServers(
+			TEXT("osc.servers"),
+			TEXT("Prints diagnostic information pertaining to the currently initialized OSC servers objects to the output log."),
+			FConsoleCommandDelegate::CreateStatic(
+				[]()
 				{
-					FString ToPrint = TEXT("    ") + Server->GetName();
-					ToPrint.Appendf(TEXT(" (Id: %u"), Server->GetUniqueID());
-					if (UWorld* World = Server->GetWorld())
-					{
-						ToPrint.Appendf(TEXT(", World: %s"), *World->GetName());
-					}
-
-					ToPrint.Appendf(TEXT(", IP: %s)"), *Server->GetIpAddress(true /* bIncludePort */));
-					ToPrint += Server->IsActive() ? TEXT(" [Active]") : TEXT(" [Inactive]");
+					FString LocalAddr;
+					ManagerPrivate::GetLocalHostAddress(LocalAddr);
+					UE_LOG(LogOSC, Display, TEXT("Local IP: %s"), *LocalAddr);
 
-					UE_LOG(LogOSC, Display, TEXT("%s"), *ToPrint);
-
-					const TArray<FOSCAddress> BoundPatterns = Server->GetBoundOSCAddressPatterns();
-					if (BoundPatterns.Num() > 0)
+					UE_LOG(LogOSC, Display, TEXT("OSC Servers:"));
+					for (TObjectIterator<UOSCServer> Itr; Itr; ++Itr)
 					{
-						UE_LOG(LogOSC, Display, TEXT("    Bound Address Patterns:"));
-						for (const FOSCAddress& Pattern : BoundPatterns)
+						if (UOSCServer* Server = *Itr)
 						{
-							UE_LOG(LogOSC, Display, TEXT("         %s"), *Pattern.GetFullPath());
+							FString ToPrint = TEXT("    ") + Server->GetName();
+							ToPrint.Appendf(TEXT(" (Id: %u"), Server->GetUniqueID());
+							if (UWorld* World = Server->GetWorld())
+							{
+								ToPrint.Appendf(TEXT(", World: %s"), *World->GetName());
+							}
+
+							ToPrint.Appendf(TEXT(", IP: %s)"), *Server->GetIpAddress(true /* bIncludePort */));
+							ToPrint += Server->IsActive() ? TEXT(" [Active]") : TEXT(" [Inactive]");
+
+							UE_LOG(LogOSC, Display, TEXT("%s"), *ToPrint);
+
+							const TArray<FOSCAddress> BoundPatterns = Server->GetBoundOSCAddressPatterns();
+							if (BoundPatterns.Num() > 0)
+							{
+								UE_LOG(LogOSC, Display, TEXT("    Bound Address Patterns:"));
+								for (const FOSCAddress& Pattern : BoundPatterns)
+								{
+									UE_LOG(LogOSC, Display, TEXT("         %s"), *Pattern.GetFullPath());
+								}
+								UE_LOG(LogOSC, Display, TEXT(""));
+							}
 						}
-						UE_LOG(LogOSC, Display, TEXT(""));
 					}
 				}
-			}
-		}
-	)
-);
-
-static FAutoConsoleCommand GOSCServerConnect(
-	TEXT("osc.server.connect"),
-	TEXT("Connects or reconnects the osc mix server with the provided name\n"
-		"(see \"osc.servers\" for a list of available servers and their respective names). Args:\n"
-		"Name - Object name of server to (re)connect\n"
-		"Address - IP Address to connect to (default: LocalHost)\n"
-		"Port - Port to connect to (default: 8095)"),
-	FConsoleCommandWithArgsDelegate::CreateStatic(
-		[](const TArray<FString>& Args)
-		{
-			FString SrvName;
-			if (Args.Num() > 0)
-			{
-				SrvName = Args[0];
-			}
+			)
+		);
+
+		static FAutoConsoleCommand GOSCServerConnect(
+			TEXT("osc.server.connect"),
+			TEXT("Connects or reconnects the osc mix server with the provided name\n"
+				"(see \"osc.servers\" for a list of available servers and their respective names). Args:\n"
+				"Name - Object name of server to (re)connect\n"
+				"Address - IP Address to connect to (default: LocalHost)\n"
+				"Port - Port to connect to (default: 8095)"),
+			FConsoleCommandWithArgsDelegate::CreateStatic(
+				[](const TArray<FString>& Args)
+				{
+					FString SrvName;
+					if (Args.Num() > 0)
+					{
+						SrvName = Args[0];
+					}
 
-			FString IPAddr;
-			OSC::GetLocalHostAddress(IPAddr);
-			if (Args.Num() > 1)
-			{
-				IPAddr = Args[1];
-			}
+					FString IPAddr;
+					GetLocalHostAddress(IPAddr);
+					if (Args.Num() > 1)
+					{
+						IPAddr = Args[1];
+					}
 
-			int32 Port = OSC::DefaultServerPort;
-			if (Args.Num() > 2)
-			{
-				Port = FCString::Atoi(*Args[2]);
-			}
+					int32 Port = DefaultServerPort;
+					if (Args.Num() > 2)
+					{
+						Port = FCString::Atoi(*Args[2]);
+					}
 
-			for (TObjectIterator<UOSCServer> Itr; Itr; ++Itr)
-			{
-				if (UOSCServer* Server = *Itr)
-				{
-					if (Server->GetName() == SrvName)
+					for (TObjectIterator<UOSCServer> Itr; Itr; ++Itr)
 					{
-						Server->Stop();
-						if (Server->SetAddress(IPAddr, Port))
+						if (UOSCServer* Server = *Itr)
 						{
-							Server->Listen();
+							if (Server->GetName() == SrvName)
+							{
+								Server->Stop();
+								if (Server->SetAddress(IPAddr, Port))
+								{
+									Server->Listen();
+								}
+								return;
+							}
 						}
-						return;
 					}
+
+					UE_LOG(LogOSC, Warning, TEXT("Server object with name '%s' not found, (re)connect not performed."), *SrvName);
 				}
-			}
+			)
+		);
+
+		static FAutoConsoleCommand GOSCServerConnectById(
+			TEXT("osc.server.connectById"),
+			TEXT("Connects or reconnects the osc mix server with the provided object id\n"
+				"(see \"osc.servers\" for a list of available servers and their respective ids). Args:\n"
+				"Id - Object Id of client to (re)connect\n"
+				"Address - IP Address to (re)connect to (default: LocalHost)\n"
+				"Port - Port to (re)connect to (default: 8095)"),
+			FConsoleCommandWithArgsDelegate::CreateStatic(
+				[](const TArray<FString>& Args)
+				{
+					if (Args.Num() == 0)
+					{
+						return;
+					}
 
-			UE_LOG(LogOSC, Warning, TEXT("Server object with name '%s' not found, (re)connect not performed."), *SrvName);
-		}
-	)
-);
-
-static FAutoConsoleCommand GOSCServerConnectById(
-	TEXT("osc.server.connectById"),
-	TEXT("Connects or reconnects the osc mix server with the provided object id\n"
-		"(see \"osc.servers\" for a list of available servers and their respective ids). Args:\n"
-		"Id - Object Id of client to (re)connect\n"
-		"Address - IP Address to (re)connect to (default: LocalHost)\n"
-		"Port - Port to (re)connect to (default: 8095)"),
-	FConsoleCommandWithArgsDelegate::CreateStatic(
-		[](const TArray<FString>& Args)
-		{
-			if (Args.Num() == 0)
-			{
-				return;
-			}
+					const int32 SrvId = FCString::Atoi(*Args[0]);
 
-			const int32 SrvId = FCString::Atoi(*Args[0]);
+					FString IPAddr;
+					GetLocalHostAddress(IPAddr);
+					if (Args.Num() > 1)
+					{
+						IPAddr = Args[1];
+					}
 
-			FString IPAddr;
-			OSC::GetLocalHostAddress(IPAddr);
-			if (Args.Num() > 1)
-			{
-				IPAddr = Args[1];
-			}
+					int32 Port = DefaultServerPort;
+					if (Args.Num() > 2)
+					{
+						Port = FCString::Atoi(*Args[2]);
+					}
 
-			int32 Port = OSC::DefaultServerPort;
-			if (Args.Num() > 2)
-			{
-				Port = FCString::Atoi(*Args[2]);
-			}
+					for (TObjectIterator<UOSCServer> Itr; Itr; ++Itr)
+					{
+						if (UOSCServer* Server = *Itr)
+						{
+							if (Server->GetUniqueID() == SrvId)
+							{
+								Server->Stop();
+								if (Server->SetAddress(IPAddr, Port))
+								{
+									Server->Listen();
+								}
+								return;
+							}
+						}
+					}
 
-			for (TObjectIterator<UOSCServer> Itr; Itr; ++Itr)
-			{
-				if (UOSCServer* Server = *Itr)
+					UE_LOG(LogOSC, Warning, TEXT("Server object with id '%u' not found, (re)connect not performed."), SrvId);
+				}
+			)
+		);
+
+		static FAutoConsoleCommand GOSCPrintClients(
+			TEXT("osc.clients"),
+			TEXT("Prints diagnostic information pertaining to the currently initialized OSC client objects to the output log."),
+			FConsoleCommandDelegate::CreateStatic(
+				[]()
 				{
-					if (Server->GetUniqueID() == SrvId)
+					FString LocalAddr;
+					GetLocalHostAddress(LocalAddr);
+					UE_LOG(LogOSC, Display, TEXT("Local IP: %s"), *LocalAddr);
+
+					UE_LOG(LogOSC, Display, TEXT("OSC Clients:"));
+					for (TObjectIterator<UOSCClient> Itr; Itr; ++Itr)
 					{
-						Server->Stop();
-						if (Server->SetAddress(IPAddr, Port))
+						if (UOSCClient* Client = *Itr)
 						{
-							Server->Listen();
+							FString ToPrint = TEXT("    ") + Client->GetName();
+							ToPrint.Appendf(TEXT(" (Id: %u"), Client->GetUniqueID());
+							if (UWorld* World = Client->GetWorld())
+							{
+								ToPrint.Appendf(TEXT(", World: %s"), *World->GetName());
+							}
+
+							FString IPAddrStr;
+							int32 Port;
+							Client->GetSendIPAddress(IPAddrStr, Port);
+							ToPrint += TEXT(", Send IP: ") + IPAddrStr + TEXT(":");
+							ToPrint.AppendInt(Port);
+							ToPrint += Client->IsActive() ? TEXT(") [Active]") : TEXT(") [Inactive]");
+
+							UE_LOG(LogOSC, Display, TEXT("%s"), *ToPrint);
 						}
-						return;
 					}
 				}
-			}
-
-			UE_LOG(LogOSC, Warning, TEXT("Server object with id '%u' not found, (re)connect not performed."), SrvId);
-		}
-	)
-);
-
-static FAutoConsoleCommand GOSCPrintClients(
-	TEXT("osc.clients"),
-	TEXT("Prints diagnostic information pertaining to the currently initialized OSC client objects to the output log."),
-	FConsoleCommandDelegate::CreateStatic(
-		[]()
-		{
-			FString LocalAddr;
-			OSC::GetLocalHostAddress(LocalAddr);
-			UE_LOG(LogOSC, Display, TEXT("Local IP: %s"), *LocalAddr);
-
-			UE_LOG(LogOSC, Display, TEXT("OSC Clients:"));
-			for (TObjectIterator<UOSCClient> Itr; Itr; ++Itr)
-			{
-				if (UOSCClient* Client = *Itr)
+			)
+		);
+
+		static FAutoConsoleCommand GOSCClientConnect(
+			TEXT("osc.client.connect"),
+			TEXT("Connects (or reconnects) the osc mix client with the provided name\n"
+				"(see \"osc.clients\" for a list of available clients and their respective ids). Args:\n"
+				"Name - Object name of client to (re)connect\n"
+				"Address - IP Address to (re)connect to (default: LocalHost)\n"
+				"Port - Port to (re)connect to (default: 8094)"),
+			FConsoleCommandWithArgsDelegate::CreateStatic(
+				[](const TArray< FString >& Args)
 				{
-					FString ToPrint = TEXT("    ") + Client->GetName();
-					ToPrint.Appendf(TEXT(" (Id: %u"), Client->GetUniqueID());
-					if (UWorld* World = Client->GetWorld())
+					if (Args.Num() == 0)
 					{
-						ToPrint.Appendf(TEXT(", World: %s"), *World->GetName());
+						return;
 					}
+					const FString CliName = Args[0];
 
-					FString IPAddrStr;
-					int32 Port;
-					Client->GetSendIPAddress(IPAddrStr, Port);
-					ToPrint += TEXT(", Send IP: ") + IPAddrStr + TEXT(":");
-					ToPrint.AppendInt(Port);
-					ToPrint += Client->IsActive() ? TEXT(") [Active]") : TEXT(") [Inactive]");
-
-					UE_LOG(LogOSC, Display, TEXT("%s"), *ToPrint);
-				}
-			}
-		}
-	)
-);
-
-static FAutoConsoleCommand GOSCClientConnect(
-	TEXT("osc.client.connect"),
-	TEXT("Connects (or reconnects) the osc mix client with the provided name\n"
-		"(see \"osc.clients\" for a list of available clients and their respective ids). Args:\n"
-		"Name - Object name of client to (re)connect\n"
-		"Address - IP Address to (re)connect to (default: LocalHost)\n"
-		"Port - Port to (re)connect to (default: 8094)"),
-	FConsoleCommandWithArgsDelegate::CreateStatic(
-		[](const TArray< FString >& Args)
-		{
-			if (Args.Num() == 0)
-			{
-				return;
-			}
-			const FString CliName = Args[0];
+					FString IPAddr;
+					GetLocalHostAddress(IPAddr);
+					if (Args.Num() > 1)
+					{
+						IPAddr = Args[1];
+					}
 
-			FString IPAddr;
-			OSC::GetLocalHostAddress(IPAddr);
-			if (Args.Num() > 1)
-			{
-				IPAddr = Args[1];
-			}
+					int32 Port = DefaultClientPort;
+					if (Args.Num() > 2)
+					{
+						Port = FCString::Atoi(*Args[2]);
+					}
 
-			int32 Port = OSC::DefaultClientPort;
-			if (Args.Num() > 2)
-			{
-				Port = FCString::Atoi(*Args[2]);
-			}
+					for (TObjectIterator<UOSCClient> Itr; Itr; ++Itr)
+					{
+						if (UOSCClient* Client = *Itr)
+						{
+							if (Client->GetName() == CliName)
+							{
+								Client->Connect();
+								Client->SetSendIPAddress(IPAddr, Port);
+								return;
+							}
+						}
+					}
 
-			for (TObjectIterator<UOSCClient> Itr; Itr; ++Itr)
-			{
-				if (UOSCClient* Client = *Itr)
+					UE_LOG(LogOSC, Warning, TEXT("Client object with name '%s' not found, (re)connect not performed."), *CliName);
+				}
+			)
+		);
+
+		static FAutoConsoleCommand GOSCClientConnectById(
+			TEXT("osc.client.connectById"),
+			TEXT("Connects (or reconnects) the osc mix client with the provided object id\n"
+				"(see \"osc.clients\" for a list of available clients and their respective ids). Args:\n"
+				"Id - Object Id of client to (re)connect\n"
+				"Address - IP Address to (re)connect to (default: LocalHost)\n"
+				"Port - Port to (re)connect to (default: 8094)"),
+			FConsoleCommandWithArgsDelegate::CreateStatic(
+				[](const TArray< FString >& Args)
 				{
-					if (Client->GetName() == CliName)
+					int32 CliId = INDEX_NONE;
+					if (Args.Num() > 0)
 					{
-						Client->Connect();
-						Client->SetSendIPAddress(IPAddr, Port);
-						return;
+						CliId = FCString::Atoi(*Args[0]);
 					}
-				}
-			}
-
-			UE_LOG(LogOSC, Warning, TEXT("Client object with name '%s' not found, (re)connect not performed."), *CliName);
-		}
-	)
-);
-
-static FAutoConsoleCommand GOSCClientConnectById(
-	TEXT("osc.client.connectById"),
-	TEXT("Connects (or reconnects) the osc mix client with the provided object id\n"
-		"(see \"osc.clients\" for a list of available clients and their respective ids). Args:\n"
-		"Id - Object Id of client to (re)connect\n"
-		"Address - IP Address to (re)connect to (default: LocalHost)\n"
-		"Port - Port to (re)connect to (default: 8094)"),
-	FConsoleCommandWithArgsDelegate::CreateStatic(
-		[](const TArray< FString >& Args)
-		{
-			int32 CliId = INDEX_NONE;
-			if (Args.Num() > 0)
-			{
-				CliId = FCString::Atoi(*Args[0]);
-			}
 
-			FString IPAddr;
-			OSC::GetLocalHostAddress(IPAddr);
-			if (Args.Num() > 1)
-			{
-				IPAddr = Args[1];
-			}
+					FString IPAddr;
+					GetLocalHostAddress(IPAddr);
+					if (Args.Num() > 1)
+					{
+						IPAddr = Args[1];
+					}
 
-			int32 Port = OSC::DefaultClientPort;
-			if (Args.Num() > 2)
-			{
-				Port = FCString::Atoi(*Args[2]);
-			}
+					int32 Port = DefaultClientPort;
+					if (Args.Num() > 2)
+					{
+						Port = FCString::Atoi(*Args[2]);
+					}
 
-			for (TObjectIterator<UOSCClient> Itr; Itr; ++Itr)
-			{
-				if (UOSCClient* Client = *Itr)
-				{
-					if (Client->GetUniqueID() == CliId)
+					for (TObjectIterator<UOSCClient> Itr; Itr; ++Itr)
 					{
-						Client->Connect();
-						Client->SetSendIPAddress(IPAddr, Port);
-						return;
+						if (UOSCClient* Client = *Itr)
+						{
+							if (Client->GetUniqueID() == CliId)
+							{
+								Client->Connect();
+								Client->SetSendIPAddress(IPAddr, Port);
+								return;
+							}
+						}
 					}
+
+					UE_LOG(LogOSC, Warning, TEXT("Client object with id '%u' not found, (re)connect not performed."), CliId);
 				}
-			}
+			)
+		);
+	} // namespace ManagerPrivate
 
-			UE_LOG(LogOSC, Warning, TEXT("Client object with id '%u' not found, (re)connect not performed."), CliId);
-		}
-	)
-);
+	int32 GetDefaultClientPort()
+	{
+		return ManagerPrivate::DefaultClientPort;
+	}
+
+	int32 GetDefaultServerPort()
+	{
+		return ManagerPrivate::DefaultServerPort;
+	}
+} // namespace UE::OSC
 
 UOSCServer* UOSCManager::CreateOSCServer(FString InReceiveIPAddress, int32 InPort, bool bInMulticastLoopback, bool bInStartListening, FString ServerName, UObject* Outer)
 {
-	if (OSC::GetLocalHostAddress(InReceiveIPAddress))
+	using namespace UE::OSC;
+
+	if (ManagerPrivate::GetLocalHostAddress(InReceiveIPAddress))
 	{
 		UE_LOG(LogOSC, Display, TEXT("OSCServer ReceiveAddress not specified. Using LocalHost IP: '%s'"), *InReceiveIPAddress);
 	}
@@ -389,7 +415,9 @@ UOSCServer* UOSCManager::CreateOSCServer(FString InReceiveIPAddress, int32 InPor
 
 UOSCClient* UOSCManager::CreateOSCClient(FString InSendIPAddress, int32 InPort, FString ClientName, UObject* Outer)
 {
-	if (OSC::GetLocalHostAddress(InSendIPAddress))
+	using namespace UE::OSC;
+
+	if (ManagerPrivate::GetLocalHostAddress(InSendIPAddress))
 	{
 		UE_LOG(LogOSC, Display, TEXT("OSCClient SendAddress not specified. Using LocalHost IP: '%s'"), *InSendIPAddress);
 	}
@@ -426,114 +454,123 @@ UOSCClient* UOSCManager::CreateOSCClient(FString InSendIPAddress, int32 InPort,
 
 FOSCMessage& UOSCManager::ClearMessage(FOSCMessage& OutMessage)
 {
-	const TSharedPtr<FOSCMessagePacket>& Packet = StaticCastSharedPtr<FOSCMessagePacket>(OutMessage.GetPacket());
-	if (Packet.IsValid())
-	{
-		Packet->GetArguments().Reset();
-	}
+	using namespace UE::OSC;
+
+	const TSharedRef<FMessagePacket>& Packet = StaticCastSharedRef<FMessagePacket>(OutMessage.GetPacketRef());
+	Packet->EmptyArguments();
 
 	return OutMessage;
 }
 
 FOSCBundle& UOSCManager::ClearBundle(FOSCBundle& OutBundle)
 {
-	const TSharedPtr<FOSCBundlePacket>& Packet = StaticCastSharedPtr<FOSCBundlePacket>(OutBundle.GetPacket());
-	if (Packet.IsValid())
-	{
-		Packet->GetPackets().Reset();
-	}
+	using namespace UE::OSC;
+
+	const TSharedRef<FBundlePacket>& Packet = StaticCastSharedRef<FBundlePacket>(OutBundle.GetPacketRef());
+	Packet->GetPackets().Reset();
 
 	return OutBundle;
 }
 
 FOSCBundle& UOSCManager::AddMessageToBundle(const FOSCMessage& InMessage, FOSCBundle& Bundle)
 {
-	const TSharedPtr<FOSCBundlePacket>& BundlePacket = StaticCastSharedPtr<FOSCBundlePacket>(Bundle.GetPacket());
-	const TSharedPtr<FOSCMessagePacket>& MessagePacket = StaticCastSharedPtr<FOSCMessagePacket>(InMessage.GetPacket());
+	using namespace UE::OSC;
 
-	if (MessagePacket.IsValid() && BundlePacket.IsValid())
-	{
-		BundlePacket->GetPackets().Add(MessagePacket);
-	}
+	const TSharedRef<FBundlePacket>& BundlePacket = StaticCastSharedRef<FBundlePacket>(Bundle.GetPacketRef());
+	const TSharedRef<FMessagePacket>& MessagePacket = StaticCastSharedRef<FMessagePacket>(InMessage.GetPacketRef());
+
+	BundlePacket->GetPackets().Add(MessagePacket);
 
 	return Bundle;
 }
 
 FOSCBundle& UOSCManager::AddBundleToBundle(const FOSCBundle& InBundle, FOSCBundle& OutBundle)
 {
-	const TSharedPtr<FOSCBundlePacket>& InBundlePacket = StaticCastSharedPtr<FOSCBundlePacket>(InBundle.GetPacket());
-	const TSharedPtr<FOSCBundlePacket>& OutBundlePacket = StaticCastSharedPtr<FOSCBundlePacket>(OutBundle.GetPacket());
+	using namespace UE::OSC;
 
-	if (InBundlePacket.IsValid() && OutBundlePacket.IsValid())
-	{
-		InBundlePacket->GetPackets().Add(OutBundlePacket);
-	}
+	const TSharedRef<FBundlePacket>& InBundlePacket = StaticCastSharedRef<FBundlePacket>(InBundle.GetPacketRef());
+	const TSharedRef<FBundlePacket>& OutBundlePacket = StaticCastSharedRef<FBundlePacket>(OutBundle.GetPacketRef());
+
+	InBundlePacket->GetPackets().Add(OutBundlePacket);
 
 	return OutBundle;
 }
 
 FOSCMessage& UOSCManager::AddFloat(FOSCMessage& OutMessage, float InValue)
 {
-	const TSharedPtr<FOSCMessagePacket>& MessagePacket = StaticCastSharedPtr<FOSCMessagePacket>(OutMessage.GetPacket());
-	MessagePacket->GetArguments().Add(FOSCType(InValue));
+	using namespace UE::OSC;
+
+	const TSharedRef<FMessagePacket>& MessagePacket = StaticCastSharedRef<FMessagePacket>(OutMessage.GetPacketRef());
+	MessagePacket->AddArgument(FOSCData(InValue));
 	return OutMessage;
 }
 
 FOSCMessage& UOSCManager::AddInt32(FOSCMessage& OutMessage, int32 InValue)
 {
-	const TSharedPtr<FOSCMessagePacket>& MessagePacket = StaticCastSharedPtr<FOSCMessagePacket>(OutMessage.GetPacket());
-	MessagePacket->GetArguments().Add(FOSCType(InValue));
+	using namespace UE::OSC;
+
+	const TSharedRef<FMessagePacket>& MessagePacket = StaticCastSharedRef<FMessagePacket>(OutMessage.GetPacketRef());
+	MessagePacket->AddArgument(FOSCData(InValue));
 	return OutMessage;
 }
 
 FOSCMessage& UOSCManager::AddInt64(FOSCMessage& OutMessage, int64 InValue)
 {
-	const TSharedPtr<FOSCMessagePacket>& MessagePacket = StaticCastSharedPtr<FOSCMessagePacket>(OutMessage.GetPacket());
-	MessagePacket->GetArguments().Add(FOSCType(InValue));
+	using namespace UE::OSC;
+
+	const TSharedRef<FMessagePacket>& MessagePacket = StaticCastSharedRef<FMessagePacket>(OutMessage.GetPacketRef());
+	MessagePacket->AddArgument(FOSCData(InValue));
 	return OutMessage;
 }
 
 FOSCMessage& UOSCManager::AddAddress(FOSCMessage& OutMessage, const FOSCAddress& InValue)
 {
-	const TSharedPtr<FOSCMessagePacket>& MessagePacket = StaticCastSharedPtr<FOSCMessagePacket>(OutMessage.GetPacket());
-	MessagePacket->GetArguments().Add(FOSCType(InValue.GetFullPath()));
+	using namespace UE::OSC;
+
+	const TSharedRef<FMessagePacket>& MessagePacket = StaticCastSharedRef<FMessagePacket>(OutMessage.GetPacketRef());
+	MessagePacket->AddArgument(FOSCData(InValue.GetFullPath()));
 	return OutMessage;
 }
 
 FOSCMessage& UOSCManager::AddString(FOSCMessage& OutMessage, FString InValue)
 {
-	const TSharedPtr<FOSCMessagePacket>& MessagePacket = StaticCastSharedPtr<FOSCMessagePacket>(OutMessage.GetPacket());
-	MessagePacket->GetArguments().Add(FOSCType(InValue));
+	using namespace UE::OSC;
+
+	const TSharedRef<FMessagePacket>& MessagePacket = StaticCastSharedRef<FMessagePacket>(OutMessage.GetPacketRef());
+	MessagePacket->AddArgument(FOSCData(InValue));
 	return OutMessage;
 }
 
 FOSCMessage& UOSCManager::AddBlob(FOSCMessage& OutMessage, const TArray<uint8>& OutValue)
 {
-	const TSharedPtr<FOSCMessagePacket>& MessagePacket = StaticCastSharedPtr<FOSCMessagePacket>(OutMessage.GetPacket());
-	MessagePacket->GetArguments().Add(FOSCType(OutValue));
+	using namespace UE::OSC;
+
+	const TSharedRef<FMessagePacket>& MessagePacket = StaticCastSharedRef<FMessagePacket>(OutMessage.GetPacketRef());
+	MessagePacket->AddArgument(FOSCData(OutValue));
 	return OutMessage;
 }
 
 FOSCMessage& UOSCManager::AddBool(FOSCMessage& OutMessage, bool InValue)
 {
-	const TSharedPtr<FOSCMessagePacket>& MessagePacket = StaticCastSharedPtr<FOSCMessagePacket>(OutMessage.GetPacket());
-	MessagePacket->GetArguments().Add(FOSCType(InValue));
+	using namespace UE::OSC;
+
+	const TSharedRef<FMessagePacket>& MessagePacket = StaticCastSharedRef<FMessagePacket>(OutMessage.GetPacketRef());
+	MessagePacket->AddArgument(FOSCData(InValue));
 	return OutMessage;
 }
 
 TArray<FOSCBundle> UOSCManager::GetBundlesFromBundle(const FOSCBundle& InBundle)
 {
+	using namespace UE::OSC;
+
 	TArray<FOSCBundle> Bundles;
-	if (InBundle.GetPacket().IsValid())
+	const TSharedRef<FBundlePacket>& BundlePacket = StaticCastSharedRef<FBundlePacket>(InBundle.GetPacketRef());
+	for (int32 i = 0; i < BundlePacket->GetPackets().Num(); i++)
 	{
-		const TSharedPtr<FOSCBundlePacket>& BundlePacket = StaticCastSharedPtr<FOSCBundlePacket>(InBundle.GetPacket());
-		for (int32 i = 0; i < BundlePacket->GetPackets().Num(); i++)
+		const TSharedRef<UE::OSC::IPacket>& Packet = BundlePacket->GetPackets()[i];
+		if (Packet->IsBundle())
 		{
-			const TSharedPtr<IOSCPacket>& Packet = BundlePacket->GetPackets()[i];
-			if (Packet->IsBundle())
-			{
-				Bundles.Emplace(StaticCastSharedPtr<FOSCMessagePacket>(Packet));
-			}
+			Bundles.Add(FOSCBundle(Packet));
 		}
 	}
 
@@ -542,21 +579,20 @@ TArray<FOSCBundle> UOSCManager::GetBundlesFromBundle(const FOSCBundle& InBundle)
 
 FOSCMessage UOSCManager::GetMessageFromBundle(const FOSCBundle& InBundle, int32 InIndex, bool& bOutSucceeded)
 {
-	if (InBundle.GetPacket().IsValid())
+	using namespace UE::OSC;
+
+	const TSharedRef<FBundlePacket>& BundlePacket = StaticCastSharedRef<FBundlePacket>(InBundle.GetPacketRef());
+	int32 Count = 0;
+	for (const TSharedRef<UE::OSC::IPacket>& Packet : BundlePacket->GetPackets())
 	{
-		const TSharedPtr<FOSCBundlePacket>& BundlePacket = StaticCastSharedPtr<FOSCBundlePacket>(InBundle.GetPacket());
-		int32 Count = 0;
-		for (const TSharedPtr<IOSCPacket>& Packet : BundlePacket->GetPackets())
+		if (Packet->IsMessage())
 		{
-			if (Packet->IsMessage())
+			if (InIndex == Count)
 			{
-				if (InIndex == Count)
-				{
-					bOutSucceeded = true;
-					return FOSCMessage(StaticCastSharedPtr<FOSCMessagePacket>(Packet));
-				}
-				Count++;
+				bOutSucceeded = true;
+				return FOSCMessage(Packet);
 			}
+			Count++;
 		}
 	}
 
@@ -564,19 +600,18 @@ FOSCMessage UOSCManager::GetMessageFromBundle(const FOSCBundle& InBundle, int32
 	return FOSCMessage();
 }
 
-TArray<FOSCMessage> UOSCManager::GetMessagesFromBundle(const FOSCBundle& OutBundle)
+TArray<FOSCMessage> UOSCManager::GetMessagesFromBundle(const FOSCBundle& InBundle)
 {
+	using namespace UE::OSC;
+
 	TArray<FOSCMessage> Messages;
-	if (OutBundle.GetPacket().IsValid())
+	const TSharedRef<FBundlePacket>& BundlePacket = StaticCastSharedRef<FBundlePacket>(InBundle.GetPacketRef());
+	for (int32 i = 0; i < BundlePacket->GetPackets().Num(); i++)
 	{
-		const TSharedPtr<FOSCBundlePacket>& BundlePacket = StaticCastSharedPtr<FOSCBundlePacket>(OutBundle.GetPacket());
-		for (int32 i = 0; i < BundlePacket->GetPackets().Num(); i++)
+		const TSharedRef<UE::OSC::IPacket>& Packet = BundlePacket->GetPackets()[i];
+		if (Packet->IsMessage())
 		{
-			const TSharedPtr<IOSCPacket>& Packet = BundlePacket->GetPackets()[i];
-			if (Packet->IsMessage())
-			{
-				Messages.Emplace(StaticCastSharedPtr<FOSCMessagePacket>(Packet));
-			}
+			Messages.Add(FOSCMessage(Packet));
 		}
 	}
 	
@@ -585,14 +620,16 @@ TArray<FOSCMessage> UOSCManager::GetMessagesFromBundle(const FOSCBundle& OutBund
 
 bool UOSCManager::GetAddress(const FOSCMessage& InMessage, const int32 InIndex, FOSCAddress& OutValue)
 {
-	if (const FOSCType* OSCType = OSC::GetOSCTypeAtIndex(InMessage, InIndex))
+	using namespace UE::OSC;
+
+	if (const FOSCData* OSCData = ManagerPrivate::GetDataAtIndex(InMessage, InIndex))
 	{
-		if (OSCType->IsString())
+		if (OSCData->IsString())
 		{
-			OutValue = FOSCAddress(OSCType->GetString());
+			OutValue = FOSCAddress(OSCData->GetString());
 			return OutValue.IsValidPath();
 		}
-		OSC_LOG_INVALID_TYPE_AT_INDEX(TEXT("String (OSCAddress)"), InIndex, InMessage);
+		ManagerPrivate::LogInvalidTypeAtIndex(EDataType::String, InIndex, InMessage);
 	}
 
 	OutValue = FOSCAddress();
@@ -601,36 +638,34 @@ bool UOSCManager::GetAddress(const FOSCMessage& InMessage, const int32 InIndex,
 
 void UOSCManager::GetAllAddresses(const FOSCMessage& InMessage, TArray<FOSCAddress>& OutValues)
 {
-	if (InMessage.GetPacket().IsValid())
+	using namespace UE::OSC;
+
+	ManagerPrivate::IterateMessageArgs(InMessage, [&](const FOSCData& OSCData)
 	{
-		const TSharedPtr<FOSCMessagePacket>& MessagePacket = StaticCastSharedPtr<FOSCMessagePacket>(InMessage.GetPacket());
-		const TArray<FOSCType>& Args = MessagePacket->GetArguments();
-		for (int32 i = 0; i < Args.Num(); i++)
+		if (OSCData.IsString())
 		{
-			const FOSCType& OSCType = Args[i];
-			if (OSCType.IsString())
+			FOSCAddress AddressToAdd = FOSCAddress(OSCData.GetString());
+			if (AddressToAdd.IsValidPath())
 			{
-				FOSCAddress AddressToAdd = FOSCAddress(OSCType.GetString());
-				if (AddressToAdd.IsValidPath())
-				{
-					OutValues.Add(MoveTemp(AddressToAdd));
-				}
+				OutValues.Add(MoveTemp(AddressToAdd));
 			}
 		}
-	}
+	});
 }
 
 bool UOSCManager::GetFloat(const FOSCMessage& InMessage, const int32 InIndex, float& OutValue)
 {
+	using namespace UE::OSC;
+
 	OutValue = 0.0f;
-	if (const FOSCType* OSCType = OSC::GetOSCTypeAtIndex(InMessage, InIndex))
+	if (const FOSCData* OSCData = ManagerPrivate::GetDataAtIndex(InMessage, InIndex))
 	{
-		if (OSCType->IsFloat())
+		if (OSCData->IsFloat())
 		{
-			OutValue = OSCType->GetFloat();
+			OutValue = OSCData->GetFloat();
 			return true;
 		}
-		OSC_LOG_INVALID_TYPE_AT_INDEX(TEXT("Float"), InIndex, InMessage);
+		ManagerPrivate::LogInvalidTypeAtIndex(EDataType::Float, InIndex, InMessage);
 	}
 
 	return false;
@@ -638,32 +673,30 @@ bool UOSCManager::GetFloat(const FOSCMessage& InMessage, const int32 InIndex, fl
 
 void UOSCManager::GetAllFloats(const FOSCMessage& InMessage, TArray<float>& OutValues)
 {
-	if (InMessage.GetPacket().IsValid())
+	using namespace UE::OSC;
+
+	ManagerPrivate::IterateMessageArgs(InMessage, [&](const FOSCData& OSCData)
 	{
-		const TSharedPtr<FOSCMessagePacket>& MessagePacket = StaticCastSharedPtr<FOSCMessagePacket>(InMessage.GetPacket());
-		const TArray<FOSCType>& Args = MessagePacket->GetArguments();
-		for (int32 i = 0; i < Args.Num(); i++)
+		if (OSCData.IsFloat())
 		{
-			const FOSCType& OSCType = Args[i];
-			if (OSCType.IsFloat())
-			{
-				OutValues.Add(OSCType.GetFloat());
-			}
+			OutValues.Add(OSCData.GetFloat());
 		}
-	}
+	});
 }
 
 bool UOSCManager::GetInt32(const FOSCMessage& InMessage, const int32 InIndex, int32& OutValue)
 {
+	using namespace UE::OSC;
+
 	OutValue = 0;
-	if (const FOSCType* OSCType = OSC::GetOSCTypeAtIndex(InMessage, InIndex))
+	if (const FOSCData* OSCData = ManagerPrivate::GetDataAtIndex(InMessage, InIndex))
 	{
-		if (OSCType->IsInt32())
+		if (OSCData->IsInt32())
 		{
-			OutValue = OSCType->GetInt32();
+			OutValue = OSCData->GetInt32();
 			return true;
 		}
-		OSC_LOG_INVALID_TYPE_AT_INDEX(TEXT("Int32"), InIndex, InMessage);
+		ManagerPrivate::LogInvalidTypeAtIndex(EDataType::Int32, InIndex, InMessage);
 	}
 
 	return false;
@@ -671,32 +704,30 @@ bool UOSCManager::GetInt32(const FOSCMessage& InMessage, const int32 InIndex, in
 
 void UOSCManager::GetAllInt32s(const FOSCMessage& InMessage, TArray<int32>& OutValues)
 {
-	if (InMessage.GetPacket().IsValid())
+	using namespace UE::OSC;
+
+	ManagerPrivate::IterateMessageArgs(InMessage, [&](const FOSCData& OSCData)
 	{
-		const TSharedPtr<FOSCMessagePacket>& MessagePacket = StaticCastSharedPtr<FOSCMessagePacket>(InMessage.GetPacket());
-		const TArray<FOSCType>& Args = MessagePacket->GetArguments();
-		for (int32 i = 0; i < Args.Num(); i++)
+		if (OSCData.IsInt32())
 		{
-			const FOSCType& OSCType = Args[i];
-			if (OSCType.IsInt32())
-			{
-				OutValues.Add(OSCType.GetInt32());
-			}
+			OutValues.Add(OSCData.GetInt32());
 		}
-	}
+	});
 }
 
 bool UOSCManager::GetInt64(const FOSCMessage& InMessage, const int32 InIndex, int64& OutValue)
 {
+	using namespace UE::OSC;
+
 	OutValue = 0l;
-	if (const FOSCType* OSCType = OSC::GetOSCTypeAtIndex(InMessage, InIndex))
+	if (const FOSCData* OSCData = ManagerPrivate::GetDataAtIndex(InMessage, InIndex))
 	{
-		if (OSCType->IsInt64())
+		if (OSCData->IsInt64())
 		{
-			OutValue = OSCType->GetInt64();
+			OutValue = OSCData->GetInt64();
 			return true;
 		}
-		OSC_LOG_INVALID_TYPE_AT_INDEX(TEXT("Int64"), InIndex, InMessage);
+		ManagerPrivate::LogInvalidTypeAtIndex(EDataType::Int64, InIndex, InMessage);
 	}
 
 	return false;
@@ -704,31 +735,29 @@ bool UOSCManager::GetInt64(const FOSCMessage& InMessage, const int32 InIndex, in
 
 void UOSCManager::GetAllInt64s(const FOSCMessage& InMessage, TArray<int64>& OutValues)
 {
-	const TSharedPtr<FOSCMessagePacket>& MessagePacket = StaticCastSharedPtr<FOSCMessagePacket>(InMessage.GetPacket());
-	if (MessagePacket.IsValid())
+	using namespace UE::OSC;
+
+	ManagerPrivate::IterateMessageArgs(InMessage, [&](const FOSCData& OSCData)
 	{
-		const TArray<FOSCType>& Args = MessagePacket->GetArguments();
-		for (int32 i = 0; i < Args.Num(); i++)
+		if (OSCData.IsInt64())
 		{
-			const FOSCType& OSCType = Args[i];
-			if (OSCType.IsInt64())
-			{
-				OutValues.Add(OSCType.GetInt64());
-			}
+			OutValues.Add(OSCData.GetInt64());
 		}
-	}
+	});
 }
 
 bool UOSCManager::GetString(const FOSCMessage& InMessage, const int32 InIndex, FString& OutValue)
 {
-	if (const FOSCType* OSCType = OSC::GetOSCTypeAtIndex(InMessage, InIndex))
+	using namespace UE::OSC;
+
+	if (const FOSCData* OSCData = ManagerPrivate::GetDataAtIndex(InMessage, InIndex))
 	{
-		if (OSCType->IsString())
+		if (OSCData->IsString())
 		{
-			OutValue = OSCType->GetString();
+			OutValue = OSCData->GetString();
 			return true;
 		}
-		OSC_LOG_INVALID_TYPE_AT_INDEX(TEXT("String"), InIndex, InMessage);
+		ManagerPrivate::LogInvalidTypeAtIndex(EDataType::String, InIndex, InMessage);
 	}
 
 	OutValue.Reset();
@@ -737,32 +766,30 @@ bool UOSCManager::GetString(const FOSCMessage& InMessage, const int32 InIndex, F
 
 void UOSCManager::GetAllStrings(const FOSCMessage& InMessage, TArray<FString>& OutValues)
 {
-	const TSharedPtr<FOSCMessagePacket>& MessagePacket = StaticCastSharedPtr<FOSCMessagePacket>(InMessage.GetPacket());
-	if (MessagePacket.IsValid())
+	using namespace UE::OSC;
+
+	ManagerPrivate::IterateMessageArgs(InMessage, [&](const FOSCData& OSCData)
 	{
-		const TArray<FOSCType>& Args = MessagePacket->GetArguments();
-		for (int32 i = 0; i < Args.Num(); i++)
+		if (OSCData.IsString())
 		{
-			const FOSCType& OSCType = Args[i];
-			if (OSCType.IsString())
-			{
-				OutValues.Add(OSCType.GetString());
-			}
+			OutValues.Add(OSCData.GetString());
 		}
-	}
+	});
 }
 
 bool UOSCManager::GetBool(const FOSCMessage& InMessage, const int32 InIndex, bool& OutValue)
 {
+	using namespace UE::OSC;
+
 	OutValue = false;
-	if (const FOSCType* OSCType = OSC::GetOSCTypeAtIndex(InMessage, InIndex))
+	if (const FOSCData* OSCData = ManagerPrivate::GetDataAtIndex(InMessage, InIndex))
 	{
-		if (OSCType->IsBool())
+		if (OSCData->IsBool())
 		{
-			OutValue = OSCType->GetBool();
+			OutValue = OSCData->GetBool();
 			return true;
 		}
-		OSC_LOG_INVALID_TYPE_AT_INDEX(TEXT("Bool"), InIndex, InMessage);
+		ManagerPrivate::LogInvalidTypeAtIndex(EDataType::True, InIndex, InMessage);
 	}
 
 	return false;
@@ -770,32 +797,30 @@ bool UOSCManager::GetBool(const FOSCMessage& InMessage, const int32 InIndex, boo
 
 void UOSCManager::GetAllBools(const FOSCMessage& InMessage, TArray<bool>& OutValues)
 {
-	const TSharedPtr<FOSCMessagePacket>& MessagePacket = StaticCastSharedPtr<FOSCMessagePacket>(InMessage.GetPacket());
-	if (MessagePacket.IsValid())
+	using namespace UE::OSC;
+
+	ManagerPrivate::IterateMessageArgs(InMessage, [&](const FOSCData& OSCData)
 	{
-		const TArray<FOSCType>& Args = MessagePacket->GetArguments();
-		for (int32 i = 0; i < Args.Num(); i++)
+		if (OSCData.IsBool())
 		{
-			const FOSCType& OSCType = Args[i];
-			if (OSCType.IsBool())
-			{
-				OutValues.Add(OSCType.GetBool());
-			}
+			OutValues.Add(OSCData.GetBool());
 		}
-	}
+	});
 }
 
 bool UOSCManager::GetBlob(const FOSCMessage& InMessage, const int32 InIndex, TArray<uint8>& OutValue)
 {
+	using namespace UE::OSC;
+
 	OutValue.Reset();
-	if (const FOSCType* OSCType = OSC::GetOSCTypeAtIndex(InMessage, InIndex))
+	if (const FOSCData* OSCData = ManagerPrivate::GetDataAtIndex(InMessage, InIndex))
 	{
-		if (OSCType->IsBlob())
+		if (OSCData->IsBlob())
 		{
-			OutValue = OSCType->GetBlob();
+			OutValue = OSCData->GetBlob();
 			return true;
 		}
-		OSC_LOG_INVALID_TYPE_AT_INDEX(TEXT("Blob"), InIndex, InMessage);
+		ManagerPrivate::LogInvalidTypeAtIndex(EDataType::Blob, InIndex, InMessage);
 	}
 
 	return false;
diff --git a/Engine/Plugins/Runtime/OSC/Source/OSC/Private/OSCMessage.cpp b/Engine/Plugins/Runtime/OSC/Source/OSC/Private/OSCMessage.cpp
--- a/Engine/Plugins/Runtime/OSC/Source/OSC/Private/OSCMessage.cpp
+++ b/Engine/Plugins/Runtime/OSC/Source/OSC/Private/OSCMessage.cpp
@@ -6,34 +6,53 @@
 
 
 FOSCMessage::FOSCMessage()
-	: Packet(MakeShared<FOSCMessagePacket>())
+	: Packet(MakeShared<UE::OSC::FMessagePacket>())
 {
 }
 
-FOSCMessage::FOSCMessage(const TSharedPtr<IOSCPacket>& InPacket)
+FOSCMessage::FOSCMessage(const TSharedRef<UE::OSC::IPacket>& InPacket)
 	: Packet(InPacket)
 {
 }
 
-FOSCMessage::~FOSCMessage()
+FOSCMessage::FOSCMessage(const TSharedPtr<UE::OSC::IPacket>& InPacket)
+	: FOSCMessage::FOSCMessage()
 {
-	Packet.Reset();
+	Packet = TSharedRef<UE::OSC::IPacket>(InPacket.Get());
 }
 
-void FOSCMessage::SetPacket(TSharedPtr<IOSCPacket>& InPacket)
+void FOSCMessage::SetPacket(TSharedPtr<UE::OSC::IPacket>& InPacket)
+{
+	using namespace UE::OSC;
+	Packet = TSharedRef<IPacket>(InPacket.Get());
+}
+
+void FOSCMessage::SetPacket(TSharedRef<UE::OSC::IPacket>& InPacket)
 {
 	Packet = InPacket;
 }
 
-const TSharedPtr<IOSCPacket>& FOSCMessage::GetPacket() const
+const TSharedPtr<UE::OSC::IPacket>& FOSCMessage::GetPacket() const
+{
+	static TSharedPtr<UE::OSC::IPacket> RetPacketPtr;
+	RetPacketPtr = TSharedPtr<UE::OSC::IPacket>(&Packet.Get());
+	return RetPacketPtr;
+}
+
+const TSharedRef<UE::OSC::IPacket>& FOSCMessage::GetPacketRef() const
 {
 	return Packet;
 }
 
+const TArray<UE::OSC::FOSCData>& FOSCMessage::GetArgumentsChecked() const
+{
+	using namespace UE::OSC;
+	return StaticCastSharedRef<FMessagePacket>(Packet)->GetArguments();
+}
 
 bool FOSCMessage::SetAddress(const FOSCAddress& InAddress)
 {
-	check(Packet.IsValid());
+	using namespace UE::OSC;
 
 	if (!InAddress.IsValidPath())
 	{
@@ -41,12 +60,13 @@ bool FOSCMessage::SetAddress(const FOSCAddress& InAddress)
 		return false;
 	}
 
-	StaticCastSharedPtr<FOSCMessagePacket>(Packet)->SetAddress(InAddress);
+	StaticCastSharedRef<UE::OSC::FMessagePacket>(Packet)->SetAddress(InAddress);
 	return true;
 }
 
 const FOSCAddress& FOSCMessage::GetAddress() const
 {
-	check(Packet.IsValid());
-	return StaticCastSharedPtr<FOSCMessagePacket>(Packet)->GetAddress();
+	using namespace UE::OSC;
+
+	return StaticCastSharedRef<FMessagePacket>(Packet)->GetAddress();
 }
diff --git a/Engine/Plugins/Runtime/OSC/Source/OSC/Private/OSCMessagePacket.cpp b/Engine/Plugins/Runtime/OSC/Source/OSC/Private/OSCMessagePacket.cpp
--- a/Engine/Plugins/Runtime/OSC/Source/OSC/Private/OSCMessagePacket.cpp
+++ b/Engine/Plugins/Runtime/OSC/Source/OSC/Private/OSCMessagePacket.cpp
@@ -6,179 +6,205 @@
 #include "OSCTypes.h"
 
 
-FOSCMessagePacket::FOSCMessagePacket()
-	: IOSCPacket()
-{}
+namespace UE::OSC
+{
+	FPacketBase::FPacketBase(const FIPv4Endpoint& InEndpoint)
+		: IPEndpoint(InEndpoint)
+	{
+	}
 
-FOSCMessagePacket::~FOSCMessagePacket()
-{}
+	const FIPv4Endpoint& FPacketBase::GetIPEndpoint() const
+	{
+		return IPEndpoint;
+	}
 
-const FOSCAddress& FOSCMessagePacket::GetAddress() const
-{
-	return Address;
-}
+	FMessagePacket::FMessagePacket(const FIPv4Endpoint& InEndpoint)
+		: FPacketBase(InEndpoint)
+	{
+	}
 
-void FOSCMessagePacket::SetAddress(const FOSCAddress& InAddress)
-{
-	Address = InAddress;
-}
+	void FMessagePacket::AddArgument(FOSCData&& InOSCData)
+	{
+		Arguments.Add(MoveTemp(InOSCData));
+	}
 
-TArray<FOSCType>& FOSCMessagePacket::GetArguments()
-{
-	return Arguments;
-}
+	void FMessagePacket::EmptyArguments()
+	{
+		Arguments.Empty();
+	}
 
-bool FOSCMessagePacket::IsBundle()
-{
-	return false;
-}
+	const FOSCAddress& FMessagePacket::GetAddress() const
+	{
+		return Address;
+	}
 
-bool FOSCMessagePacket::IsMessage()
-{
-	return true;
-}
+	const TArray<FOSCData>& FMessagePacket::GetArguments() const
+	{
+		return Arguments;
+	}
 
-void FOSCMessagePacket::WriteData(FOSCStream& Stream)
-{
-	if (!Address.IsValidPath())
+	void FMessagePacket::SetAddress(const FOSCAddress& InAddress)
 	{
-		UE_LOG(LogOSC, Warning, TEXT("Failed to write OSCMessagePacket. Invalid OSCAddress '%s'"), *Address.GetFullPath());
-		return;
+		Address = InAddress;
 	}
 
-	// Begin writing data 
-	Stream.WriteString(Address.GetFullPath());
+	void FMessagePacket::SetAddress(FOSCAddress&& InAddress)
+	{
+		Address = MoveTemp(InAddress);
+	}
 
-	// write type tags
-	FString TagTypes = ",";
-	for (int32 i = 0; i < Arguments.Num(); i++)
+	bool FMessagePacket::IsBundle()
 	{
-		TagTypes += static_cast<ANSICHAR>(Arguments[i].GetTypeTag());
+		return false;
 	}
 
-	Stream.WriteString(TagTypes);
+	bool FMessagePacket::IsMessage()
+	{
+		return true;
+	}
 
-	for (FOSCType OSCType : Arguments)
+	void FMessagePacket::WriteData(FStream& Stream)
 	{
-		switch (OSCType.GetTypeTag())
-		{
-		case EOSCTypeTag::OSC_CHAR:
-			Stream.WriteChar(OSCType.GetChar());
-			break;
-		case EOSCTypeTag::OSC_INT32:
-			Stream.WriteInt32(OSCType.GetInt32());
-			break;
-		case EOSCTypeTag::OSC_FLOAT:
-			Stream.WriteFloat(OSCType.GetFloat());
-			break;
-		case EOSCTypeTag::OSC_DOUBLE:
-			Stream.WriteDouble(OSCType.GetDouble());
-			break;
-		case EOSCTypeTag::OSC_INT64:
-			Stream.WriteInt64(OSCType.GetInt64());
-			break;
-		case EOSCTypeTag::OSC_TIME:
-			Stream.WriteUInt64(OSCType.GetTimeTag());
-			break;
-		case EOSCTypeTag::OSC_STRING:
-			Stream.WriteString(OSCType.GetString());
-			break;
-		case EOSCTypeTag::OSC_BLOB:
+		if (!Address.IsValidPath())
 		{
-			TArray<uint8> blob = OSCType.GetBlob();
-			Stream.WriteBlob(blob);
-		}
-		break;
-		case EOSCTypeTag::OSC_COLOR:
-#if PLATFORM_LITTLE_ENDIAN
-			Stream.WriteInt32(OSCType.GetColor().ToPackedABGR());
-#else
-			Stream.WriteInt32(OSCType.GetColor().ToPackedRGBA());
-#endif
-			break;
-		case EOSCTypeTag::OSC_TRUE:
-		case EOSCTypeTag::OSC_FALSE:
-		case EOSCTypeTag::OSC_NIL:
-		case EOSCTypeTag::OSC_INFINITUM:
-			// No values are written for these types
-			break;
-		default:
-			// Argument is not supported 
-			unimplemented();
-			break;
+			UE_LOG(LogOSC, Warning, TEXT("Failed to write OSCMessagePacket. Invalid OSCAddress '%s'"), *Address.GetFullPath());
+			return;
 		}
-	}
-}
 
-void FOSCMessagePacket::ReadData(FOSCStream& Stream)
-{
-	// Read Address
-	Address = Stream.ReadString();
+		// Begin writing data 
+		Stream.WriteString(Address.GetFullPath());
 
-	// Read string of tags
-	const FString StreamString = Stream.ReadString();
+		// write type tags
+		FString TagTypes = ",";
+		for (int32 i = 0; i < Arguments.Num(); i++)
+		{
+			TagTypes += static_cast<ANSICHAR>(Arguments[i].GetDataType());
+		}
 
-	const TArray<TCHAR, FString::AllocatorType>& TagTypes = StreamString.GetCharArray();
-	if(TagTypes.Num() == 0)
-	{
-		UE_LOG(LogOSC, Error, TEXT("Failed to read message packet with address '%s' from stream: Invalid (Empty) Type Tag"), *Address.GetFullPath());
-		return;
+		Stream.WriteString(TagTypes);
+
+		for (const FOSCData& OSCData : Arguments)
+		{
+			switch (OSCData.GetDataType())
+			{
+			case EDataType::Char:
+				Stream.WriteChar(OSCData.GetChar());
+				break;
+			case EDataType::Int32:
+				Stream.WriteInt32(OSCData.GetInt32());
+				break;
+			case EDataType::Float:
+				Stream.WriteFloat(OSCData.GetFloat());
+				break;
+			case EDataType::Double:
+				Stream.WriteDouble(OSCData.GetDouble());
+				break;
+			case EDataType::Int64:
+				Stream.WriteInt64(OSCData.GetInt64());
+				break;
+			case EDataType::Time:
+				Stream.WriteUInt64(OSCData.GetTimeTag());
+				break;
+			case EDataType::String:
+				Stream.WriteString(OSCData.GetString());
+				break;
+			case EDataType::Blob:
+			{
+				TArray<uint8> blob = OSCData.GetBlob();
+				Stream.WriteBlob(blob);
+			}
+			break;
+			case EDataType::Color:
+	#if PLATFORM_LITTLE_ENDIAN
+				Stream.WriteInt32(OSCData.GetColor().ToPackedABGR());
+	#else
+				Stream.WriteInt32(OSCData.GetColor().ToPackedRGBA());
+	#endif
+				break;
+			case EDataType::True:
+			case EDataType::False:
+			case EDataType::Nil:
+			case EDataType::Infinitum:
+				// No values are written for these types
+				break;
+			default:
+				// Argument is not supported 
+				unimplemented();
+				break;
+			}
+		}
 	}
 
-	// Skip the first argument which is ','
-	for (int32 i = 1; i < TagTypes.Num(); i++)
+	void FMessagePacket::ReadData(FStream& Stream)
 	{
-		const EOSCTypeTag Tag = static_cast<EOSCTypeTag>(TagTypes[i]);
-		switch (Tag)
+		// Read Address
+		Address = Stream.ReadString();
+
+		// Read string of tags
+		const FString StreamString = Stream.ReadString();
+
+		const TArray<TCHAR, FString::AllocatorType>& TagTypes = StreamString.GetCharArray();
+		if(TagTypes.Num() == 0)
 		{
-		case EOSCTypeTag::OSC_CHAR:
-			Arguments.Add(FOSCType(Stream.ReadChar()));
-			break;
-		case EOSCTypeTag::OSC_INT32:
-			Arguments.Add(FOSCType(Stream.ReadInt32()));
-			break;
-		case EOSCTypeTag::OSC_FLOAT:
-			Arguments.Add(FOSCType(Stream.ReadFloat()));
-			break;
-		case EOSCTypeTag::OSC_DOUBLE:
-			Arguments.Add(FOSCType(Stream.ReadDouble()));
-			break;
-		case EOSCTypeTag::OSC_INT64:
-			Arguments.Add(FOSCType(Stream.ReadInt64()));
-			break;
-		case EOSCTypeTag::OSC_TRUE:
-			Arguments.Add(FOSCType(true));
-			break;
-		case EOSCTypeTag::OSC_FALSE:
-			Arguments.Add(FOSCType(false));
-			break;
-		case EOSCTypeTag::OSC_NIL:
-			Arguments.Add(FOSCType(EOSCTypeTag::OSC_NIL));
-			break;
-		case EOSCTypeTag::OSC_INFINITUM:
-			Arguments.Add(FOSCType(EOSCTypeTag::OSC_INFINITUM));
-			break;
-		case EOSCTypeTag::OSC_TIME:
-			Arguments.Add(FOSCType(Stream.ReadUInt64()));
-			break;
-		case EOSCTypeTag::OSC_STRING:
-			Arguments.Add(FOSCType(Stream.ReadString()));
-			break;
-		case EOSCTypeTag::OSC_BLOB:
-			Arguments.Add(FOSCType(Stream.ReadBlob()));
-			break;
-		case EOSCTypeTag::OSC_COLOR:
-			Arguments.Add(FOSCType(FColor(Stream.ReadInt32())));
-			break;
-		case EOSCTypeTag::OSC_TERMINATE:
-			// Return on first terminate found. FString GetCharArray can return
-			// an array with multiple terminators.
+			UE_LOG(LogOSC, Error, TEXT("Failed to read message packet with address '%s' from stream: Invalid (Empty) Type Tag"), *Address.GetFullPath());
 			return;
+		}
 
-		default:
-			// Argument is not supported 
-			unimplemented();
-			break;
+		// Skip the first argument which is ','
+		for (int32 i = 1; i < TagTypes.Num(); i++)
+		{
+			const EDataType Tag = static_cast<EDataType>(TagTypes[i]);
+			switch (Tag)
+			{
+			case EDataType::Char:
+				Arguments.Add(FOSCData(Stream.ReadChar()));
+				break;
+			case EDataType::Int32:
+				Arguments.Add(FOSCData(Stream.ReadInt32()));
+				break;
+			case EDataType::Float:
+				Arguments.Add(FOSCData(Stream.ReadFloat()));
+				break;
+			case EDataType::Double:
+				Arguments.Add(FOSCData(Stream.ReadDouble()));
+				break;
+			case EDataType::Int64:
+				Arguments.Add(FOSCData(Stream.ReadInt64()));
+				break;
+			case EDataType::True:
+				Arguments.Add(FOSCData(true));
+				break;
+			case EDataType::False:
+				Arguments.Add(FOSCData(false));
+				break;
+			case EDataType::Nil:
+				Arguments.Add(FOSCData::Nil());
+				break;
+			case EDataType::Infinitum:
+				Arguments.Add(FOSCData::Infinitum());
+				break;
+			case EDataType::Time:
+				Arguments.Add(FOSCData(Stream.ReadUInt64()));
+				break;
+			case EDataType::String:
+				Arguments.Add(FOSCData(Stream.ReadString()));
+				break;
+			case EDataType::Blob:
+				Arguments.Add(FOSCData(Stream.ReadBlob()));
+				break;
+			case EDataType::Color:
+				Arguments.Add(FOSCData(FColor(Stream.ReadInt32())));
+				break;
+			case EDataType::Terminate:
+				// Return on first terminate found. FString GetCharArray can return
+				// an array with multiple terminators.
+				return;
+
+			default:
+				// Argument is not supported 
+				unimplemented();
+				break;
+			}
 		}
 	}
-}
+} // namespace UE::OSC
diff --git a/Engine/Plugins/Runtime/OSC/Source/OSC/Private/OSCMessagePacket.h b/Engine/Plugins/Runtime/OSC/Source/OSC/Private/OSCMessagePacket.h
--- a/Engine/Plugins/Runtime/OSC/Source/OSC/Private/OSCMessagePacket.h
+++ b/Engine/Plugins/Runtime/OSC/Source/OSC/Private/OSCMessagePacket.h
@@ -7,37 +7,70 @@
 class FOSCType;
 
 
-class FOSCMessagePacket : public IOSCPacket
+namespace UE::OSC
 {
-public:
-	FOSCMessagePacket();
-	virtual ~FOSCMessagePacket();
+	class FPacketBase : public IPacket
+	{
+	public:
+		FPacketBase() = default;
+		FPacketBase(const FIPv4Endpoint& InEndpoint);
 
-	/** Set OSC message address. */
-	void SetAddress(const FOSCAddress& InAddress);
+		virtual ~FPacketBase() = default;
 
-	/** Get OSC message address. */
-	virtual const FOSCAddress& GetAddress() const;
+		/** Get endpoint IP address and port responsible for creation/forwarding of packet */
+		virtual const FIPv4Endpoint& GetIPEndpoint() const override;
 
-	/** Get arguments array. */
-	TArray<FOSCType>& GetArguments();
+		UE_DEPRECATED(5.4, "Packets now have address and port set using endpoints.")
+		virtual const FString& GetIPAddress() const override { static const FString None; return None; }
 
-	/** Returns false to indicate type is not OSC bundle. */
-	virtual bool IsBundle();
+		UE_DEPRECATED(5.4, "Packets now have port and port set using endpoints.")
+		virtual uint16 GetPort() const override { return 0; };
 
-	/** Returns true to indicate its an OSC message. */
-	virtual bool IsMessage();
+	protected:
+		FIPv4Endpoint IPEndpoint;
+	};
 
-	/** Write message data into an OSC stream. */
-	virtual void WriteData(FOSCStream& Stream) override;
 
-	/** Reads message data from an OSC stream and creates new argument. */
-	virtual void ReadData(FOSCStream& Stream) override;
+	class FMessagePacket : public FPacketBase
+	{
+	public:
+		FMessagePacket() = default;
+		FMessagePacket(const FIPv4Endpoint& InEndpoint);
+		virtual ~FMessagePacket() = default;
 
-private:
-	/** OSC address. */
-	FOSCAddress Address;
+		/** Adds argument to argument array */
+		void AddArgument(FOSCData&& InOSCType);
 
-	/** List of argument types. */
-	TArray<FOSCType> Arguments;
-};
+		/** Empties all arguments */
+		void EmptyArguments();
+
+		/** Set OSC message address. */
+		void SetAddress(const FOSCAddress& InAddress);
+		void SetAddress(FOSCAddress&& InAddress);
+
+		/** Get OSC message address. */
+		virtual const FOSCAddress& GetAddress() const;
+
+		/** Get arguments array. */
+		virtual const TArray<FOSCData>& GetArguments() const;
+
+		/** Returns false to indicate type is not OSC bundle. */
+		virtual bool IsBundle();
+
+		/** Returns true to indicate its an OSC message. */
+		virtual bool IsMessage();
+
+		/** Write message data into an OSC stream. */
+		virtual void WriteData(FStream& Stream) override;
+
+		/** Reads message data from an OSC stream and creates new argument. */
+		virtual void ReadData(FStream& Stream) override;
+
+	private:
+		/** OSC address. */
+		FOSCAddress Address;
+
+		/** List of argument data types. */
+		TArray<FOSCData> Arguments;
+	};
+} // namespace UE::OSC
diff --git a/Engine/Plugins/Runtime/OSC/Source/OSC/Private/OSCModule.cpp b/Engine/Plugins/Runtime/OSC/Source/OSC/Private/OSCModule.cpp
new file mode 100644
--- /dev/null
+++ b/Engine/Plugins/Runtime/OSC/Source/OSC/Private/OSCModule.cpp
@@ -0,0 +1,21 @@
+// Copyright Epic Games, Inc. All Rights Reserved.
+
+#include "OSCModule.h"
+
+#include "OSCLog.h"
+
+DEFINE_LOG_CATEGORY(LogOSC);
+
+
+namespace UE::OSC
+{
+	void FModule::StartupModule()
+	{
+		if (!FModuleManager::Get().LoadModule(TEXT("Networking")))
+		{
+			UE_LOG(LogOSC, Error, TEXT("Required module 'Networking' failed to load. OSC service disabled."));
+		}
+	}
+} // namespace UE::OSC
+
+IMPLEMENT_MODULE(UE::OSC::FModule, OSC)
diff --git a/Engine/Plugins/Runtime/OSC/Source/OSC/Private/OSCModule.h b/Engine/Plugins/Runtime/OSC/Source/OSC/Private/OSCModule.h
new file mode 100644
--- /dev/null
+++ b/Engine/Plugins/Runtime/OSC/Source/OSC/Private/OSCModule.h
@@ -0,0 +1,15 @@
+// Copyright Epic Games, Inc. All Rights Reserved.
+
+#pragma once
+
+#include "Modules/ModuleManager.h"
+
+
+namespace UE::OSC
+{
+	class FModule : public IModuleInterface
+	{
+	public:
+		virtual void StartupModule() override;
+	};
+} // namespace UE::OSC
diff --git a/Engine/Plugins/Runtime/OSC/Source/OSC/Private/OSCPacket.cpp b/Engine/Plugins/Runtime/OSC/Source/OSC/Private/OSCPacket.cpp
--- a/Engine/Plugins/Runtime/OSC/Source/OSC/Private/OSCPacket.cpp
+++ b/Engine/Plugins/Runtime/OSC/Source/OSC/Private/OSCPacket.cpp
@@ -6,37 +6,28 @@
 #include "OSCLog.h"
 
 
-TSharedPtr<IOSCPacket> IOSCPacket::CreatePacket(const uint8* InPacketType, const FString& InIPAddress, uint16 InPort)
+namespace UE::OSC
 {
-	const FString PacketIdentifier(ANSI_TO_TCHAR((const ANSICHAR*)&InPacketType[0]));
-	
-	TSharedPtr<IOSCPacket> Packet;
-	if (PacketIdentifier.StartsWith(OSC::PathSeparator))
-	{
-		Packet = MakeShared<FOSCMessagePacket>();
-	}
-	else if (PacketIdentifier == OSC::BundleTag)
+	TSharedPtr<IPacket> IPacket::CreatePacket(const uint8* InPacketType, const FIPv4Endpoint& InIPEndpoint)
 	{
-		Packet = MakeShared<FOSCBundlePacket>();
-	}
-	else
-	{
-		UE_LOG(LogOSC, Warning, TEXT("Failed to parse lead character of OSC packet. "
-			"Lead identifier of '%c' not valid bundle tag ('%s') or message ('%s') identifier."), PacketIdentifier[0], *OSC::BundleTag, *OSC::PathSeparator);
-		return nullptr;
-	}
-
-	Packet->IPAddress = InIPAddress;
-	Packet->Port = InPort;
-	return Packet;
-}
-
-const FString& IOSCPacket::GetIPAddress() const
-{
-	return IPAddress;
-}
+		const FString PacketIdentifier(ANSI_TO_TCHAR((const ANSICHAR*)&InPacketType[0]));
+	
+		TSharedPtr<FPacketBase> Packet;
+		if (PacketIdentifier.StartsWith(OSC::PathSeparator))
+		{
+			Packet = MakeShared<FMessagePacket>(InIPEndpoint);
+		}
+		else if (PacketIdentifier == UE::OSC::BundleTag)
+		{
+			Packet = MakeShared<OSC::FBundlePacket>(InIPEndpoint);
+		}
+		else
+		{
+			UE_LOG(LogOSC, Warning, TEXT("Failed to parse lead character of OSC packet. "
+				"Lead identifier of '%c' not valid bundle tag ('%s') or message ('%s') identifier."), PacketIdentifier[0], *OSC::BundleTag, *OSC::PathSeparator);
+			return nullptr;
+		}
 
-uint16 IOSCPacket::GetPort() const
-{
-	return Port;
-}
+		return Packet;
+	}
+} // namespace UE::OSC
diff --git a/Engine/Plugins/Runtime/OSC/Source/OSC/Private/OSCServer.cpp b/Engine/Plugins/Runtime/OSC/Source/OSC/Private/OSCServer.cpp
--- a/Engine/Plugins/Runtime/OSC/Source/OSC/Private/OSCServer.cpp
+++ b/Engine/Plugins/Runtime/OSC/Source/OSC/Private/OSCServer.cpp
@@ -170,19 +170,21 @@ void UOSCServer::ClearPackets()
 	OSCPackets.Empty();
 }
 
-void UOSCServer::EnqueuePacket(TSharedPtr<IOSCPacket> InPacket)
+void UOSCServer::EnqueuePacket(TSharedPtr<UE::OSC::IPacket> InPacket)
 {
 	OSCPackets.Enqueue(InPacket);
 }
 
 void UOSCServer::DispatchBundle(const FString& InIPAddress, uint16 InPort, const FOSCBundle& InBundle)
 {
+	using namespace UE::OSC;
+
 	OnOscBundleReceived.Broadcast(InBundle, InIPAddress, InPort);
 	OnOscBundleReceivedNative.Broadcast(InBundle, InIPAddress, InPort);
 
-	TSharedPtr<FOSCBundlePacket> BundlePacket = StaticCastSharedPtr<FOSCBundlePacket>(InBundle.GetPacket());
-	FOSCBundlePacket::FPacketBundle Packets = BundlePacket->GetPackets();
-	for (TSharedPtr<IOSCPacket>& Packet : Packets)
+	TSharedPtr<FBundlePacket> BundlePacket = StaticCastSharedRef<FBundlePacket>(InBundle.GetPacketRef());
+	TArray<TSharedRef<IPacket>>& Packets = BundlePacket->GetPackets();
+	for (TSharedRef<IPacket>& Packet : Packets)
 	{
 		if (Packet->IsMessage())
 		{
@@ -222,21 +224,23 @@ void UOSCServer::DispatchMessage(const FString& InIPAddress, uint16 InPort, cons
 
 void UOSCServer::PumpPacketQueue(const TSet<uint32>* AllowlistedClients)
 {
-	TSharedPtr<IOSCPacket> Packet;
+	using namespace UE::OSC;
+
+	TSharedPtr<UE::OSC::IPacket> Packet;
 	while (OSCPackets.Dequeue(Packet))
 	{
 		FIPv4Address IPAddr;
-		const FString& Address = Packet->GetIPAddress();
-		if (!AllowlistedClients || (FIPv4Address::Parse(Address, IPAddr) && AllowlistedClients->Contains(IPAddr.Value)))
+		const FIPv4Endpoint& Endpoint = Packet->GetIPEndpoint();
+		if (!AllowlistedClients || AllowlistedClients->Contains(Endpoint.Address.Value))
 		{
-			uint16 Port = Packet->GetPort();
+			const FString StringAdr = Endpoint.Address.ToString();
 			if (Packet->IsMessage())
 			{
-				DispatchMessage(Address, Port, FOSCMessage(Packet));
+				DispatchMessage(StringAdr, Endpoint.Port, FOSCMessage(Packet.ToSharedRef()));
 			}
 			else if (Packet->IsBundle())
 			{
-				DispatchBundle(Address, Port, FOSCBundle(Packet));
+				DispatchBundle(StringAdr, Endpoint.Port, FOSCBundle(Packet.ToSharedRef()));
 			}
 			else
 			{
diff --git a/Engine/Plugins/Runtime/OSC/Source/OSC/Private/OSCServerProxy.cpp b/Engine/Plugins/Runtime/OSC/Source/OSC/Private/OSCServerProxy.cpp
--- a/Engine/Plugins/Runtime/OSC/Source/OSC/Private/OSCServerProxy.cpp
+++ b/Engine/Plugins/Runtime/OSC/Source/OSC/Private/OSCServerProxy.cpp
@@ -28,14 +28,16 @@ FOSCServerProxy::~FOSCServerProxy()
 
 void FOSCServerProxy::OnPacketReceived(const FArrayReaderPtr& InData, const FIPv4Endpoint& InEndpoint)
 {
-	TSharedPtr<IOSCPacket> Packet = IOSCPacket::CreatePacket(InData->GetData(), InEndpoint.Address.ToString(), InEndpoint.Port);
+	using namespace UE::OSC;
+
+	TSharedPtr<IPacket> Packet = IPacket::CreatePacket(InData->GetData(), InEndpoint);
 	if (!Packet.IsValid())
 	{
 		UE_LOG(LogOSC, Verbose, TEXT("Message received from endpoint '%s' invalid OSC packet."), *InEndpoint.ToString());
 		return;
 	}
 
-	FOSCStream Stream = FOSCStream(InData->GetData(), InData->Num());
+	FStream Stream = FStream(InData->GetData(), InData->Num());
 	Packet->ReadData(Stream);
 	Server->EnqueuePacket(Packet);
 }
diff --git a/Engine/Plugins/Runtime/OSC/Source/OSC/Private/OSCStream.cpp b/Engine/Plugins/Runtime/OSC/Source/OSC/Private/OSCStream.cpp
--- a/Engine/Plugins/Runtime/OSC/Source/OSC/Private/OSCStream.cpp
+++ b/Engine/Plugins/Runtime/OSC/Source/OSC/Private/OSCStream.cpp
@@ -6,487 +6,310 @@
 #include "Math/Color.h"
 
 
-namespace OSC
+namespace UE::OSC
 {
-	const int32 DefaultWriteStreamSize = 1024;
-} // namespace OSC
-
-FOSCStream::FOSCStream()
-	: Position(0)
-	, bIsReadStream(false)
-{
-	// Set default size to reasonable size to
-	// avoid thrashing with allocations while writing
-	Data.Reserve(OSC::DefaultWriteStreamSize);
-}
-
-FOSCStream::FOSCStream(const uint8* InData, int32 InSize)
-	: Data(InData, InSize)
-	, Position(0)
-	, bIsReadStream(true)
-{
-}
-
-const uint8* FOSCStream::GetData() const
-{
-	return Data.GetData();
-}
-
-int32 FOSCStream::GetLength() const
-{
-	return Data.Num();
-}
-
-bool FOSCStream::HasReachedEnd() const
-{
-	return Position >= Data.Num();
-}
-
-int32 FOSCStream::GetPosition() const
-{
-	return Position;
-}
-
-void FOSCStream::SetPosition(int32 InPosition)
-{
-	check(InPosition <= Data.Num());
-	Position = InPosition;
-}
+	FStream::FStream()
+	{
+		// Set default size to reasonable size to avoid thrashing with allocations while writing
+		Data.Reserve(1024);
+	}
 
-TCHAR FOSCStream::ReadChar()
-{
-	uint8 Temp;
-	if (Read(&Temp, 1) > 0)
+	FStream::FStream(const uint8* InData, int32 InSize)
+		: Data(InData, InSize)
+		, bIsReadStream(true)
 	{
-		UE_CLOG(Temp > 0x7F, LogOSC, Warning, TEXT("Non-ANSI character '%u' written to OSCStream"), Temp);
-		return (TCHAR)Temp;
 	}
 
-	return '\0';
-}
+	const uint8* FStream::GetData() const
+	{
+		return Data.GetData();
+	}
 
-void FOSCStream::WriteChar(TCHAR Char)
-{
-	const uint32 Temp = FChar::ToUnsigned(Char);
-	UE_CLOG(Temp > 0x7F, LogOSC, Warning, TEXT("Non-ANSI character '%u' written to OSCStream"), Temp);
-	const uint8 Temp8 = (uint8)Temp;
-	Write(&Temp8, 1);
-}
+	int32 FStream::GetLength() const
+	{
+		return Data.Num();
+	}
 
-FColor FOSCStream::ReadColor()
-{
-	uint32 Packed = static_cast<uint32>(ReadInt32());
-	return FColor(Packed);
-}
+	bool FStream::HasReachedEnd() const
+	{
+		return Position >= Data.Num();
+	}
 
-void FOSCStream::WriteColor(FColor Color)
-{
-#if PLATFORM_LITTLE_ENDIAN
-	uint32 Packed = Color.ToPackedABGR();
-#else // PLATFORM_LITTLE_ENDIAN
-	uint32 Packed = Color.ToPackedRGBA();
-#endif // !PLATFORM_LITTLE_ENDIAN
+	int32 FStream::GetPosition() const
+	{
+		return Position;
+	}
 
-	WriteInt32(static_cast<int32>(Packed));
-}
+	void FStream::SetPosition(int32 InPosition)
+	{
+		check(InPosition <= Data.Num());
+		Position = InPosition;
+	}
 
-int32 FOSCStream::ReadInt32()
-{
-	uint8 Temp[4];
-	if (Read(Temp, 4) == 4)
+	TCHAR FStream::ReadChar()
 	{
-#if PLATFORM_LITTLE_ENDIAN
-		union {
-			int32 i;
-			uint8 c[4];
-		} u;
+		uint8 Temp;
+		if (Read(&Temp, 1) > 0)
+		{
+			UE_CLOG(Temp > 0x7F, LogOSC, Warning, TEXT("Non-ANSI character '%u' written to OSCStream"), Temp);
+			return (TCHAR)Temp;
+		}
 
-		u.c[0] = Temp[3];
-		u.c[1] = Temp[2];
-		u.c[2] = Temp[1];
-		u.c[3] = Temp[0];
+		return '\0';
+	}
 
-		return u.i;
-#else
-		return *(int32*)(Temp);
-#endif // !PLATFORM_LITTLE_ENDIAN
+	void FStream::WriteChar(TCHAR Char)
+	{
+		const uint32 Temp = FChar::ToUnsigned(Char);
+		UE_CLOG(Temp > 0x7F, LogOSC, Warning, TEXT("Non-ANSI character '%u' written to OSCStream"), Temp);
+		const uint8 Temp8 = (uint8)Temp;
+		Write(&Temp8, 1);
 	}
 
-	return 0;
-}
+	FColor FStream::ReadColor()
+	{
+		uint32 Packed = static_cast<uint32>(ReadInt32());
+		return FColor(Packed);
+	}
 
-void FOSCStream::WriteInt32(int32 Value)
-{
-	uint8 Temp[4];
-
-#ifdef PLATFORM_LITTLE_ENDIAN
-	union {
-		int32 i;
-		uint8 c[4];
-	} u;
-
-	u.i = Value;
-
-	Temp[3] = u.c[0];
-	Temp[2] = u.c[1];
-	Temp[1] = u.c[2];
-	Temp[0] = u.c[3];
-#else // PLATFORM_LITTLE_ENDIAN
-	*(int32*)(Temp) = Value;
-#endif // !PLATFORM_LITTLE_ENDIAN
-	Write(Temp, 4);
-}
-
-double FOSCStream::ReadDouble()
-{
-	uint8 Temp[8];
-	if (Read(Temp, 8) == 8)
+	void FStream::WriteColor(FColor Color)
 	{
-#if PLATFORM_LITTLE_ENDIAN
-		union {
-			double d;
-			uint8 c[8];
-		} u;
+	#if PLATFORM_LITTLE_ENDIAN
+		uint32 Packed = Color.ToPackedABGR();
+	#else // PLATFORM_LITTLE_ENDIAN
+		uint32 Packed = Color.ToPackedRGBA();
+	#endif // !PLATFORM_LITTLE_ENDIAN
 
-		u.c[0] = Temp[7];
-		u.c[1] = Temp[6];
-		u.c[2] = Temp[5];
-		u.c[3] = Temp[4];
-		u.c[4] = Temp[3];
-		u.c[5] = Temp[2];
-		u.c[6] = Temp[1];
-		u.c[7] = Temp[0];
+		WriteInt32(static_cast<int32>(Packed));
+	}
 
-		return u.d;
-#else // PLATFORM_LITTLE_ENDIAN
-		return *(double*)Temp;
-#endif // !PLATFORM_LITTLE_ENDIAN
+	int32 FStream::ReadInt32()
+	{
+		return ReadNumeric<int32>();
+	}
 
+	void FStream::WriteInt32(int32 Value)
+	{
+		WriteNumeric<int32>(Value);
 	}
 
-	return 0;
-}
+	double FStream::ReadDouble()
+	{
+		return ReadNumeric<double>();
+	}
 
-void FOSCStream::WriteDouble(uint64 Value)
-{
-	uint8 Temp[8];
-
-#ifdef PLATFORM_LITTLE_ENDIAN
-	union {
-		double i;
-		uint8 c[8];
-	} u;
-
-	u.i = Value;
-
-	Temp[7] = u.c[0];
-	Temp[6] = u.c[1];
-	Temp[5] = u.c[2];
-	Temp[4] = u.c[3];
-	Temp[3] = u.c[4];
-	Temp[2] = u.c[5];
-	Temp[1] = u.c[6];
-	Temp[0] = u.c[7];
-#else // PLATFORM_LITTLE_ENDIAN
-	*(double*)(Temp) = Value;
-#endif // !PLATFORM_LITTLE_ENDIAN
-	Write(Temp, 8);
-}
-
-int64 FOSCStream::ReadInt64()
-{
-	uint8 Temp[8];
-	if (Read(Temp, 8) == 8)
+	void FStream::WriteDouble(uint64 Value)
 	{
-#if PLATFORM_LITTLE_ENDIAN
+		uint8 Temp[8];
+
+	#ifdef PLATFORM_LITTLE_ENDIAN
 		union {
-			int64 i;
+			double i;
 			uint8 c[8];
 		} u;
 
-		u.c[0] = Temp[7];
-		u.c[1] = Temp[6];
-		u.c[2] = Temp[5];
-		u.c[3] = Temp[4];
-		u.c[4] = Temp[3];
-		u.c[5] = Temp[2];
-		u.c[6] = Temp[1];
-		u.c[7] = Temp[0];
-
-		return u.i;
-#else // PLATFORM_LITTLE_ENDIAN
-		return *(int64*)Temp;
-#endif // !PLATFORM_LITTLE_ENDIAN
+		u.i = Value;
+
+		Temp[7] = u.c[0];
+		Temp[6] = u.c[1];
+		Temp[5] = u.c[2];
+		Temp[4] = u.c[3];
+		Temp[3] = u.c[4];
+		Temp[2] = u.c[5];
+		Temp[1] = u.c[6];
+		Temp[0] = u.c[7];
+	#else // PLATFORM_LITTLE_ENDIAN
+		*(double*)(Temp) = Value;
+	#endif // !PLATFORM_LITTLE_ENDIAN
+		Write(Temp, 8);
 	}
 
-	return 0;
-}
+	int64 FStream::ReadInt64()
+	{
+		return ReadNumeric<int64>();
+	}
 
-void FOSCStream::WriteInt64(int64 Value)
-{
-	uint8 Temp[8];
-
-#ifdef PLATFORM_LITTLE_ENDIAN
-	union {
-		int64 i;
-		uint8 c[8];
-	} u;
-
-	u.i = Value;
-
-	Temp[7] = u.c[0];
-	Temp[6] = u.c[1];
-	Temp[5] = u.c[2];
-	Temp[4] = u.c[3];
-	Temp[3] = u.c[4];
-	Temp[2] = u.c[5];
-	Temp[1] = u.c[6];
-	Temp[0] = u.c[7];
-#else // PLATFORM_LITTLE_ENDIAN
-	*(int64*)(Temp) = Value;
-#endif // !PLATFORM_LITTLE_ENDIAN
-	Write(Temp, 8);
-}
-
-uint64 FOSCStream::ReadUInt64()
-{
-	uint8 Temp[8];
-	if (Read(Temp, 8) == 8)
+	void FStream::WriteInt64(int64 Value)
 	{
-#if PLATFORM_LITTLE_ENDIAN
+		uint8 Temp[8];
+
+	#ifdef PLATFORM_LITTLE_ENDIAN
 		union {
-			uint64 i;
+			int64 i;
 			uint8 c[8];
 		} u;
 
-		u.c[0] = Temp[7];
-		u.c[1] = Temp[6];
-		u.c[2] = Temp[5];
-		u.c[3] = Temp[4];
-		u.c[4] = Temp[3];
-		u.c[5] = Temp[2];
-		u.c[6] = Temp[1];
-		u.c[7] = Temp[0];
-
-		return u.i;
-#else // PLATFORM_LITTLE_ENDIAN
-		return *(uint64*)Temp;
-#endif // !PLATFORM_LITTLE_ENDIAN
-
+		u.i = Value;
+
+		Temp[7] = u.c[0];
+		Temp[6] = u.c[1];
+		Temp[5] = u.c[2];
+		Temp[4] = u.c[3];
+		Temp[3] = u.c[4];
+		Temp[2] = u.c[5];
+		Temp[1] = u.c[6];
+		Temp[0] = u.c[7];
+	#else // PLATFORM_LITTLE_ENDIAN
+		*(int64*)(Temp) = Value;
+	#endif // !PLATFORM_LITTLE_ENDIAN
+		Write(Temp, 8);
 	}
 
-	return 0;
-}
-
-void FOSCStream::WriteUInt64(uint64 Value)
-{
-	uint8 Temp[8];
-
-#ifdef PLATFORM_LITTLE_ENDIAN
-	union {
-		uint64 i;
-		uint8 c[8];
-	} u;
-
-	u.i = Value;
-
-	Temp[7] = u.c[0];
-	Temp[6] = u.c[1];
-	Temp[5] = u.c[2];
-	Temp[4] = u.c[3];
-	Temp[3] = u.c[4];
-	Temp[2] = u.c[5];
-	Temp[1] = u.c[6];
-	Temp[0] = u.c[7];
-#else // PLATFORM_LITTLE_ENDIAN
-	*reinterpret_cast<uint64*>(Temp) = Value;
-#endif // !PLATFORM_LITTLE_ENDIAN
-	Write(Temp, 8);
-}
-
-float FOSCStream::ReadFloat()
-{
-	uint8 Temp[4];
-	if (Read(Temp, 4) == 4)
+	uint64 FStream::ReadUInt64()
 	{
-#if PLATFORM_LITTLE_ENDIAN
-		union {
-			float f;
-			uint8 c[4];
-		} u;
-
-		u.c[0] = Temp[3];
-		u.c[1] = Temp[2];
-		u.c[2] = Temp[1];
-		u.c[3] = Temp[0];
-
-		return u.f;
-#else // PLATFORM_LITTLE_ENDIAN
-		return *(float*)Temp;
-#endif // !PLATFORM_LITTLE_ENDIAN
+		return ReadNumeric<uint64>();
 	}
 
-	return 0.0f;
-}
-
-void FOSCStream::WriteFloat(float Value)
-{
-	uint8 Temp[4];
-
-#ifdef PLATFORM_LITTLE_ENDIAN
-	union {
-		float f;
-		uint8 c[4];
-	} u;
-
-	u.f = Value;
-
-	Temp[3] = u.c[0];
-	Temp[2] = u.c[1];
-	Temp[1] = u.c[2];
-	Temp[0] = u.c[3];
-#else // if !PLATFORM_LITTLE_ENDIAN
-	*(float*)(Temp) = Value;
-#endif // !PLATFORM_LITTLE_ENDIAN
-
-	Write(Temp, 4);
-}
-
-FString FOSCStream::ReadString()
-{
-	const int32 DataSize = Data.Num();
-	if (HasReachedEnd() || DataSize == 0)
+	void FStream::WriteUInt64(uint64 Value)
 	{
-		return FString();
+		WriteNumeric<uint64>(Value);
 	}
 
-	// Cache init position index and push along until either at end or character read is null terminator.
-	const int32 InitPosition = GetPosition();
-	while (!HasReachedEnd() && ReadChar() != '\0');
+	float FStream::ReadFloat()
+	{
+		return ReadNumeric<float>();
+	}
 
-	if (HasReachedEnd() && Data.Last() != '\0')
+	void FStream::WriteFloat(float Value)
 	{
-		UE_LOG(LogOSC, Error, TEXT("Invalid string when reading OSCStream: Null terminator '\0' not found"));
-		return FString();
+		WriteNumeric<float>(Value);
 	}
 
-	// Cache end for string copy, increment to next read
-	// location, and consume pad until next 4-byte boundary.
-	const int32 EndPosition = Position;
+	FString FStream::ReadString()
+	{
+		const int32 DataSize = Data.Num();
+		if (HasReachedEnd() || DataSize == 0)
+		{
+			return FString();
+		}
 
-	// Count includes the null terminator.
-	const int32 Count = EndPosition - InitPosition;
-	check(Count > 0);
+		// Cache init position index and push along until either at end or character read is null terminator.
+		const int32 InitPosition = GetPosition();
+		while (!HasReachedEnd() && ReadChar() != '\0');
 
-	const int32 UnboundByteCount = Count % 4;
-	if (UnboundByteCount != 0)
-	{
-		Position += 4 - UnboundByteCount;
-	}
+		if (HasReachedEnd() && Data.Last() != '\0')
+		{
+			UE_LOG(LogOSC, Error, TEXT("Invalid string when reading OSCStream: Null terminator '\0' not found"));
+			return FString();
+		}
 
-	// Exclude null terminator here; this constructor appends one.
-	return FString(Count - 1, (ANSICHAR*)(&Data[InitPosition]));
-}
+		// Cache end for string copy, increment to next read
+		// location, and consume pad until next 4-byte boundary.
+		const int32 EndPosition = Position;
 
-void FOSCStream::WriteString(const FString& InString)
-{
-	const TArray<TCHAR, FString::AllocatorType>& CharArr = InString.GetCharArray();
+		// Count includes the null terminator.
+		const int32 Count = EndPosition - InitPosition;
+		check(Count > 0);
 
-	int32 Count = CharArr.Num();
-	if (Count == 0)
-	{
-		WriteChar('\0');
-		Count++;
-	}
-	else
-	{
-		for (int32 i = 0; i < Count; i++)
+		const int32 UnboundByteCount = Count % 4;
+		if (UnboundByteCount != 0)
 		{
-			WriteChar(CharArr[i]);
+			Position += 4 - UnboundByteCount;
 		}
+
+		// Exclude null terminator here; this constructor appends one.
+		return FString(Count - 1, (ANSICHAR*)(&Data[InitPosition]));
 	}
 
-	// Increment & pad string with null terminator (String must
-	// be packed into multiple of 4 bytes)
-	check(Count > 0);
-	const int32 UnboundByteCount = Count % 4;
-	if (UnboundByteCount != 0)
+	void FStream::WriteString(const FString& InString)
 	{
-		const int32 NumPaddingZeros = 4 - UnboundByteCount;
-		for (int32 i = 0; i < NumPaddingZeros; i++)
+		const TArray<TCHAR, FString::AllocatorType>& CharArr = InString.GetCharArray();
+
+		int32 Count = CharArr.Num();
+		if (Count == 0)
 		{
 			WriteChar('\0');
+			Count++;
+		}
+		else
+		{
+			for (int32 i = 0; i < Count; i++)
+			{
+				WriteChar(CharArr[i]);
+			}
 		}
-	}
-}
 
-TArray<uint8> FOSCStream::ReadBlob()
-{
-	TArray<uint8> Blob;
+		// Increment & pad string with null terminator (String must
+		// be packed into multiple of 4 bytes)
+		check(Count > 0);
+		const int32 UnboundByteCount = Count % 4;
+		if (UnboundByteCount != 0)
+		{
+			const int32 NumPaddingZeros = 4 - UnboundByteCount;
+			for (int32 i = 0; i < NumPaddingZeros; i++)
+			{
+				WriteChar('\0');
+			}
+		}
+	}
 
-	const int32 BlobSize = ReadInt32();
-	for (int32 i = 0; i < BlobSize; i++)
+	TArray<uint8> FStream::ReadBlob()
 	{
-		Blob.Add(ReadChar());
-	}
+		TArray<uint8> Blob;
 
-	Position = ((Position + 3) / 4) * 4; // padded
+		const int32 BlobSize = ReadInt32();
+		for (int32 i = 0; i < BlobSize; i++)
+		{
+			Blob.Add(ReadChar());
+		}
 
-	return Blob;
-}
+		Position = ((Position + 3) / 4) * 4; // padded
 
-void FOSCStream::WriteBlob(TArray<uint8>& Blob)
-{
-	WriteInt32(Blob.Num());
-	for (int32 i = 0; i < Blob.Num(); i++)
-	{
-		Write(&Blob[i], 4u);
+		return Blob;
 	}
-}
 
-int32 FOSCStream::Read(uint8* Buffer, int32 InSize)
-{
-	check(bIsReadStream);
-	const int32 DataSize = Data.Num();
-	if (InSize == 0 || Position >= DataSize)
+	void FStream::WriteBlob(TArray<uint8>& Blob)
 	{
-		return 0;
+		WriteInt32(Blob.Num());
+		for (int32 i = 0; i < Blob.Num(); i++)
+		{
+			Write(&Blob[i], 4u);
+		}
 	}
 
-	const int32 Num = FMath::Min<int32>(InSize, DataSize - Position);
-	if (Num > 0)
+	int32 FStream::Read(uint8* OutBuffer, int32 InSize)
 	{
-		FMemory::Memcpy(Buffer, &Data[Position], Num);
-		Position += Num;
-	}
+		check(bIsReadStream);
+		const int32 DataSize = Data.Num();
+		if (InSize == 0 || Position >= DataSize)
+		{
+			return 0;
+		}
 
-	return Num;
-}
+		const int32 Num = FMath::Min<int32>(InSize, DataSize - Position);
+		if (Num > 0)
+		{
+			FMemory::Memcpy(OutBuffer, &Data[Position], Num);
+			Position += Num;
+		}
 
-int32 FOSCStream::Write(const uint8* InBuffer, int32 InSize)
-{
-	check(!bIsReadStream);
-	if (InSize <= 0)
-	{
-		return 0;
+		return Num;
 	}
 
-	if (Position < Data.Num())
+	int32 FStream::Write(const uint8* InBuffer, int32 InSize)
 	{
-		int32 Slack = Data.Num() - Position;
-		if (InSize - Slack > 0)
+		check(!bIsReadStream);
+		if (InSize <= 0)
 		{
-			Data.AddUninitialized(InSize - Slack);
+			return 0;
 		}
-	}
-	else
-	{
-		check(Position == Data.Num());
-		Data.AddUninitialized(InSize);
-	}
 
-	FMemory::Memcpy(&Data[Position], InBuffer, InSize);
-	Position += InSize;
-	return InSize;
-}
+		if (Position < Data.Num())
+		{
+			int32 Slack = Data.Num() - Position;
+			if (InSize - Slack > 0)
+			{
+				Data.AddUninitialized(InSize - Slack);
+			}
+		}
+		else
+		{
+			check(Position == Data.Num());
+			Data.AddUninitialized(InSize);
+		}
 
+		FMemory::Memcpy(&Data[Position], InBuffer, InSize);
+		Position += InSize;
+		return InSize;
+	}
+} // namespace UE::OSC
diff --git a/Engine/Plugins/Runtime/OSC/Source/OSC/Private/OSCTypes.cpp b/Engine/Plugins/Runtime/OSC/Source/OSC/Private/OSCTypes.cpp
new file mode 100644
--- /dev/null
+++ b/Engine/Plugins/Runtime/OSC/Source/OSC/Private/OSCTypes.cpp
@@ -0,0 +1,251 @@
+// Copyright Epic Games, Inc. All Rights Reserved.
+
+#include "OSCTypes.h"
+
+
+namespace UE::OSC
+{
+	namespace TypesPrivate
+	{
+		template <typename TVariantType>
+		TVariantType GetSafe(const FOSCData::FVariant& InVariant, const TVariantType& Default)
+		{
+			if (const TVariantType* Value = InVariant.TryGet<TVariantType>())
+			{
+				return *Value;
+			}
+
+			return Default;
+		}
+	} // namespace TypesPrivate
+
+	const TCHAR* LexToString(EDataType DataType)
+	{
+		switch (DataType)
+		{
+			case EDataType::Blob:
+				return TEXT("Blob");
+			case EDataType::Char:
+				return TEXT("Char");
+			case EDataType::Color:
+				return TEXT("Color");
+			case EDataType::Double:
+				return TEXT("Double");
+			case EDataType::Float:
+				return TEXT("Float");
+			case EDataType::Infinitum:
+				return TEXT("Infinitum");
+			case EDataType::Int32:
+				return TEXT("Int32");
+			case EDataType::Int64:
+				return TEXT("Int64");
+			case EDataType::Nil:
+				return TEXT("Nil");
+			case EDataType::String:
+				return TEXT("String");
+			case EDataType::Terminate:
+				return TEXT("Terminate");
+			case EDataType::Time:
+				return TEXT("Time");
+
+			// Both treated as bool
+			case EDataType::False:
+			case EDataType::True:
+				return TEXT("Bool");
+
+			default:
+			{
+				checkNoEntry();
+				return nullptr;
+			}
+		}
+	}
+
+	FOSCData::FOSCData(const TArray<uint8>& Value)
+		: DataType(EDataType::Blob)
+	{
+		Data.Set<TArray<uint8>>(Value);
+	}
+
+	FOSCData::FOSCData(TArray<uint8>&& Value)
+		: DataType(EDataType::Blob)
+	{
+		Data.Set<TArray<uint8>>(MoveTemp(Value));
+	}
+
+	FOSCData::FOSCData(bool Value)
+		: DataType(Value ? EDataType::True : EDataType::False)
+	{
+		Data.Set<int32>(0);
+	}
+
+	FOSCData::FOSCData(ANSICHAR Value)
+		: DataType(EDataType::Char)
+	{
+		Data.Set<ANSICHAR>(Value);
+	}
+
+	FOSCData::FOSCData(FColor Value)
+		: DataType(EDataType::Color)
+	{
+		Data.Set<FColor>(Value);
+	}
+
+	FOSCData::FOSCData(double Value)
+		: DataType(EDataType::Double)
+	{
+		Data.Set<double>(Value);
+	}
+
+	FOSCData::FOSCData(float Value)
+		: DataType(EDataType::Float)
+	{
+		Data.Set<float>(Value);
+	}
+
+	FOSCData::FOSCData(int32 Value)
+		: DataType(EDataType::Int32)
+	{
+		Data.Set<int32>(Value);
+	}
+
+	FOSCData::FOSCData(int64 Value)
+		: DataType(EDataType::Int64)
+	{
+		Data.Set<int64>(Value);
+	}
+
+	FOSCData::FOSCData(FString Value)
+		: DataType(EDataType::String)
+	{
+		Data.Set<FString>(MoveTemp(Value));
+	}
+
+	FOSCData::FOSCData(uint64 Value)
+		: DataType(EDataType::Time)
+	{
+		Data.Set<uint64>(Value);
+	}
+
+	FOSCData::FOSCData::FOSCData(EDataType DataType)
+		: DataType(DataType)
+	{
+		Data.Set<int32>(0);
+	}
+
+	const FOSCData& FOSCData::Nil()
+	{
+		auto MakeInf = []()
+			{
+				FOSCData NilType;
+				NilType.DataType = EDataType::Nil;
+				return NilType;
+			};
+		static const FOSCData StaticNil = MakeInf();
+		return StaticNil;
+	}
+
+	const FOSCData& FOSCData::Infinitum()
+	{
+		auto MakeInf = []()
+			{
+				FOSCData Inf;
+				Inf.DataType = EDataType::Infinitum;
+				return Inf;
+			};
+		static const FOSCData InfType = MakeInf();
+		return InfType;
+	}
+
+	const FOSCData& FOSCData::Terminate()
+	{
+		auto MakeTerm = []()
+			{
+				FOSCData Term;
+				Term.DataType = EDataType::Terminate;
+				return Term;
+			};
+		static const FOSCData TermType = MakeTerm();
+		return TermType;
+	}
+
+	bool FOSCData::IsNil(const FOSCData& InType)
+	{
+		return InType.GetDataType() == EDataType::Nil;
+	}
+
+	EDataType FOSCData::GetDataType() const
+	{
+		return DataType;
+	}
+
+	TArray<uint8> FOSCData::GetBlob() const
+	{
+		return TypesPrivate::GetSafe<TArray<uint8>>(Data, TArray<uint8>());
+	}
+
+	bool FOSCData::GetBool() const
+	{
+		return DataType == EDataType::True;
+	}
+
+	ANSICHAR FOSCData::GetChar() const
+	{
+		return TypesPrivate::GetSafe<ANSICHAR>(Data, '\0');
+	}
+
+	FColor FOSCData::GetColor() const
+	{
+		return TypesPrivate::GetSafe(Data, FColor::Black);
+	}
+
+	double FOSCData::GetDouble() const
+	{
+		return TypesPrivate::GetSafe<double>(Data, 0.0);
+	}
+
+	float FOSCData::GetFloat() const
+	{
+		return TypesPrivate::GetSafe<float>(Data, 0.0f);
+	}
+
+	int32 FOSCData::GetInt32() const
+	{
+		return TypesPrivate::GetSafe<int32>(Data, 0);
+	}
+
+	int64 FOSCData::GetInt64() const
+	{
+		return TypesPrivate::GetSafe<int64>(Data, 0);
+	}
+
+	FString FOSCData::GetString() const
+	{
+		return TypesPrivate::GetSafe(Data, FString());
+	}
+
+	uint64 FOSCData::GetTimeTag() const
+	{
+		return TypesPrivate::GetSafe<uint64>(Data, 0);
+	}
+
+	TArrayView<const uint8> FOSCData::GetBlobArrayView() const
+	{
+		if (const TArray<uint8>* BlobPtr = Data.TryGet<TArray<uint8>>())
+		{
+			return TArrayView<const uint8>(BlobPtr->GetData(), BlobPtr->Num());
+		}
+
+		return { };
+	}
+
+	FStringView FOSCData::GetStringView() const
+	{
+		if (const FString* StrPtr = Data.TryGet<FString>())
+		{
+			return *StrPtr;
+		}
+
+		return { };
+	}
+} // namespace UE::OSC
diff --git a/Engine/Plugins/Runtime/OSC/Source/OSC/Public/OSCAddress.h b/Engine/Plugins/Runtime/OSC/Source/OSC/Public/OSCAddress.h
--- a/Engine/Plugins/Runtime/OSC/Source/OSC/Public/OSCAddress.h
+++ b/Engine/Plugins/Runtime/OSC/Source/OSC/Public/OSCAddress.h
@@ -1,7 +1,6 @@
 // Copyright Epic Games, Inc. All Rights Reserved.
 #pragma once
 
-#include "CoreMinimal.h"
 #include "UObject/ObjectMacros.h"
 #include "UObject/Object.h"
 #include "OSCLog.h"
@@ -9,11 +8,11 @@
 #include "OSCAddress.generated.h"
 
 
-namespace OSC
+namespace UE::OSC
 {
 	extern const FString BundleTag;
 	extern const FString PathSeparator;
-} // namespace OSC
+} // namespace UE::OSC
 
 
 USTRUCT(BlueprintType)
@@ -23,11 +22,9 @@ struct OSC_API FOSCAddress
 
 private:
 	// Ordered array of container names
-	UPROPERTY(Transient)
 	TArray<FString> Containers;
 
 	// Method name of string
-	UPROPERTY(Transient)
 	FString Method;
 
 	/** Cached values for validity and hash */
@@ -53,19 +50,25 @@ public:
 	bool Matches(const FOSCAddress& InAddress) const;
 
 	/** Pushes container onto end of address' ordered array of containers */
-	void PushContainer(const FString& InContainer);
+	bool PushContainer(const FString& InContainer);
+
+	/** Pushes container onto end of address' ordered array of containers */
+	bool PushContainer(FString&& InContainer);
+
+	/** Pushes containers onto end of address' ordered array of containers */
+	bool PushContainers(const TArray<FString>& InContainers);
 
 	/** Pushes containers onto end of address' ordered array of containers */
-	void PushContainers(const TArray<FString>& InContainers);
+	bool PushContainers(TArray<FString>&& InContainers);
 
 	/** Pops container from ordered array of containers */
-	FString PopContainer();
+	FString PopContainer(bool* bDidPop = nullptr);
 
 	/** Pops containers off end of address' ordered array of containers */
-	TArray<FString> PopContainers(int32 InNumContainers);
+	TArray<FString> PopContainers(int32 InNumContainers, bool* bDidPop = nullptr);
 
 	/** Removes containers from container array at index until count */
-	void RemoveContainers(int32 InIndex, int32 InCount);
+	bool RemoveContainers(int32 InIndex, int32 InCount);
 
 	/** Clears ordered array of containers */
 	void ClearContainers();
@@ -73,8 +76,23 @@ public:
 	/** Get method name of address */
 	const FString& GetMethod() const;
 
+	/** Returns numeric prefix of OSC method if is numeric */
+	int32 GetNumericPrefix(bool* bIsNumeric = nullptr) const;
+
+	/** Returns numeric suffix of OSC method if is numeric */
+	int32 GetNumericSuffix(bool* bIsNumeric = nullptr) const;
+
+	/** Sets the address to the given array of containers and provided method */
+	bool Set(TArray<FString>&& InContainers, FString&& InMethod);
+
+	/** Sets the address to the given array of containers and provided method */
+	bool Set(const TArray<FString>& InContainers, const FString& InMethod);
+
+	/** Sets the method name of address */
+	bool SetMethod(const FString& InMethod);
+
 	/** Sets the method name of address */
-	void SetMethod(const FString& InMethod);
+	bool SetMethod(FString&& InMethod);
 
 	/** Returns container path of OSC address in the form '/Container1/Container2' */
 	FString GetContainerPath() const;
@@ -138,4 +156,7 @@ public:
 	{
 		return InAddress.Hash;
 	}
+
+private:
+	bool PushContainersInternal(TArray<FString>&& InContainers);
 };
diff --git a/Engine/Plugins/Runtime/OSC/Source/OSC/Public/OSCBundle.h b/Engine/Plugins/Runtime/OSC/Source/OSC/Public/OSCBundle.h
--- a/Engine/Plugins/Runtime/OSC/Source/OSC/Public/OSCBundle.h
+++ b/Engine/Plugins/Runtime/OSC/Source/OSC/Public/OSCBundle.h
@@ -11,12 +11,23 @@ struct OSC_API FOSCBundle
 	GENERATED_USTRUCT_BODY()
 
 	FOSCBundle();
-	FOSCBundle(const TSharedPtr<IOSCPacket>& InPacket);
-	~FOSCBundle();
+	explicit FOSCBundle(const TSharedRef<UE::OSC::IPacket>& InPacket);
 
-	void SetPacket(const TSharedPtr<IOSCPacket>& InPacket);
-	const TSharedPtr<IOSCPacket>& GetPacket() const;
+	UE_DEPRECATED(5.5, "Use explicit shared ref constructor instead")
+	FOSCBundle(const TSharedPtr<UE::OSC::IPacket>& InPacket);
+
+	~FOSCBundle() = default;
+
+	UE_DEPRECATED(5.5, "Use shared ref setter instead")
+	void SetPacket(TSharedPtr<UE::OSC::IPacket>& InPacket);
+
+	void SetPacket(const TSharedRef<UE::OSC::IPacket>& InPacket);
+
+	UE_DEPRECATED(5.5, "Use shared ref getter instead")
+	const TSharedPtr<UE::OSC::IPacket>& GetPacket() const;
+
+	const TSharedRef<UE::OSC::IPacket>& GetPacketRef() const;
 
 private:
-	TSharedPtr<IOSCPacket> Packet;
+	TSharedRef<UE::OSC::IPacket> Packet;
 };
\ No newline at end of file
diff --git a/Engine/Plugins/Runtime/OSC/Source/OSC/Public/OSCClient.h b/Engine/Plugins/Runtime/OSC/Source/OSC/Public/OSCClient.h
--- a/Engine/Plugins/Runtime/OSC/Source/OSC/Public/OSCClient.h
+++ b/Engine/Plugins/Runtime/OSC/Source/OSC/Public/OSCClient.h
@@ -1,7 +1,6 @@
 // Copyright Epic Games, Inc. All Rights Reserved.
 #pragma once
 
-#include "CoreMinimal.h"
 #include "Templates/UniquePtr.h"
 #include "UObject/Object.h"
 
@@ -11,21 +10,43 @@
 #include "OSCClient.generated.h"
 
 
-/** Interface for internal networking implementation.  See UOSCClient for details */
-class OSC_API IOSCClientProxy
+namespace UE::OSC
 {
-public:
-	virtual ~IOSCClientProxy() { }
+	/** Interface for internal network implementation of sending OSC messages & bundles as a client. */
+	class OSC_API IClientProxy
+	{
+	public:
+		// Creates a new client proxy that can be used by any system where the provided dispatch callback is called on a worker thread.
+		static TUniquePtr<IClientProxy> Create(const FString& ClientName);
+
+		virtual ~IClientProxy() = default;
+
+		UE_DEPRECATED(5.5, "Use GetSendIPEndpoint instead")
+		virtual void GetSendIPAddress(FString& InIPAddress, int32& Port) const = 0;
+
+		UE_DEPRECATED(5.5, "Use SetSendIPEndpoint instead")
+		virtual bool SetSendIPAddress(const FString& InIPAddress, const int32 Port) = 0;
 
-	virtual void GetSendIPAddress(FString& InIPAddress, int32& Port) const = 0;
-	virtual bool SetSendIPAddress(const FString& InIPAddress, const int32 Port) = 0;
+		virtual const FIPv4Endpoint& GetSendIPEndpoint() const = 0;
+		virtual void SetSendIPEndpoint(const FIPv4Endpoint& InEndpoint) = 0;
 
-	virtual bool IsActive() const = 0;
+		virtual bool IsActive() const = 0;
 
-	virtual void SendMessage(FOSCMessage& Message) = 0;
-	virtual void SendBundle(FOSCBundle& Bundle) = 0;
+		virtual void SendMessage(FOSCMessage& Message) = 0;
+		virtual void SendBundle(FOSCBundle& Bundle) = 0;
+
+		virtual void Stop() = 0;
+	};
+} // namespace UE::OSC
+
+// For backward compat.  To be deprecated
+class OSC_API IOSCClientProxy : public UE::OSC::IClientProxy
+{
+public:
+	UE_DEPRECATED(5.5, "Use UE::OSC::IClientProxy instead")
+	IOSCClientProxy() = default;
 
-	virtual void Stop() = 0;
+	virtual ~IOSCClientProxy() = default;
 };
 
 UCLASS(BlueprintType)
@@ -62,5 +83,5 @@ protected:
 	void Stop();
 	
 	/** Pointer to internal implementation of client proxy */
-	TUniquePtr<IOSCClientProxy> ClientProxy;
+	TUniquePtr<UE::OSC::IClientProxy> ClientProxy;
 };
diff --git a/Engine/Plugins/Runtime/OSC/Source/OSC/Public/OSCLog.h b/Engine/Plugins/Runtime/OSC/Source/OSC/Public/OSCLog.h
--- a/Engine/Plugins/Runtime/OSC/Source/OSC/Public/OSCLog.h
+++ b/Engine/Plugins/Runtime/OSC/Source/OSC/Public/OSCLog.h
@@ -1,6 +1,6 @@
 // Copyright Epic Games, Inc. All Rights Reserved.
 #pragma once
 
-#include "CoreMinimal.h"
+#include "Logging/LogMacros.h"
 
 DECLARE_LOG_CATEGORY_EXTERN(LogOSC, Log, All);
diff --git a/Engine/Plugins/Runtime/OSC/Source/OSC/Public/OSCManager.h b/Engine/Plugins/Runtime/OSC/Source/OSC/Public/OSCManager.h
--- a/Engine/Plugins/Runtime/OSC/Source/OSC/Public/OSCManager.h
+++ b/Engine/Plugins/Runtime/OSC/Source/OSC/Public/OSCManager.h
@@ -16,6 +16,12 @@
 class UOSCServer;
 class UOSCClient;
 
+namespace UE::OSC
+{
+	OSC_API int32 GetDefaultClientPort();
+	OSC_API int32 GetDefaultServerPort();
+}
+
 UCLASS()
 class OSC_API UOSCManager : public UBlueprintFunctionLibrary
 {
@@ -28,7 +34,7 @@ public:
 	UFUNCTION(BlueprintCallable, Category = "Audio|OSC")
 	static UOSCServer* CreateOSCServer(FString ReceiveIPAddress, int32 Port, bool bMulticastLoopback, bool bStartListening, FString ServerName, UObject* Outer = nullptr);
 
-	// Creates an OSC Client.  If SendIPAddress left empty (or '0'), attempts to use
+	// Creates an OSC Client.  If SendIPAddress left empty (or '0')
 	// attempts to use LocalHost IP address.
 	UFUNCTION(BlueprintCallable, Category = "Audio|OSC")
 	static UOSCClient* CreateOSCClient(FString SendIPAddress, int32 Port, FString ClientName, UObject* Outer = nullptr);
diff --git a/Engine/Plugins/Runtime/OSC/Source/OSC/Public/OSCMessage.h b/Engine/Plugins/Runtime/OSC/Source/OSC/Public/OSCMessage.h
--- a/Engine/Plugins/Runtime/OSC/Source/OSC/Public/OSCMessage.h
+++ b/Engine/Plugins/Runtime/OSC/Source/OSC/Public/OSCMessage.h
@@ -18,15 +18,29 @@ struct OSC_API FOSCMessage
 	GENERATED_USTRUCT_BODY()
 
 	FOSCMessage();
-	FOSCMessage(const TSharedPtr<IOSCPacket>& InPacket);
-	~FOSCMessage();
+	FOSCMessage(const TSharedRef<UE::OSC::IPacket>& InPacket);
 
-	void SetPacket(TSharedPtr<IOSCPacket>& InPacket);
-	const TSharedPtr<IOSCPacket>& GetPacket() const;
+	UE_DEPRECATED(5.4, "Use shared ref ctor instead")
+	FOSCMessage(const TSharedPtr<UE::OSC::IPacket>& InPacket);
+
+	~FOSCMessage() = default;
+
+	// Returns arguments, asserting if the message's packet is unset.
+	const TArray<UE::OSC::FOSCData>& GetArgumentsChecked() const;
+
+	UE_DEPRECATED(5.4, "Use shared ref setter instead")
+	void SetPacket(TSharedPtr<UE::OSC::IPacket>& InPacket);
+
+	void SetPacket(TSharedRef<UE::OSC::IPacket>& InPacket);
+
+	UE_DEPRECATED(5.4, "Use shared ref getter instead")
+	const TSharedPtr<UE::OSC::IPacket>& GetPacket() const;
+
+	const TSharedRef<UE::OSC::IPacket>& GetPacketRef() const;
 
 	bool SetAddress(const FOSCAddress& InAddress);
 	const FOSCAddress& GetAddress() const;
 
 private:
-	TSharedPtr<IOSCPacket> Packet;
+	TSharedRef<UE::OSC::IPacket> Packet;
 };
\ No newline at end of file
diff --git a/Engine/Plugins/Runtime/OSC/Source/OSC/Public/OSCModule.cpp b/Engine/Plugins/Runtime/OSC/Source/OSC/Public/OSCModule.cpp
deleted file mode 100644
diff --git a/Engine/Plugins/Runtime/OSC/Source/OSC/Public/OSCPacket.h b/Engine/Plugins/Runtime/OSC/Source/OSC/Public/OSCPacket.h
--- a/Engine/Plugins/Runtime/OSC/Source/OSC/Public/OSCPacket.h
+++ b/Engine/Plugins/Runtime/OSC/Source/OSC/Public/OSCPacket.h
@@ -1,42 +1,55 @@
 // Copyright Epic Games, Inc. All Rights Reserved.
 #pragma once
 
+#include "Interfaces/IPv4/IPv4Endpoint.h"
+
 #include "OSCTypes.h"
 #include "OSCStream.h"
 #include "OSCAddress.h"
 
-// Forward Declarations
-class FOSCStream;
-
 
-class OSC_API IOSCPacket
+namespace UE::OSC
 {
-public:
-	IOSCPacket() = default;
-	virtual ~IOSCPacket() = default;
+	// Forward Declarations
+	class FStream;
+
+	class OSC_API IPacket : public TSharedFromThis<IPacket>
+	{
+	public:
+		IPacket() = default;
+		virtual ~IPacket() = default;
 
-	/** Write packet data into stream */
-	virtual void WriteData(FOSCStream& OutStream) = 0;
+		/** Write packet data into stream */
+		virtual void WriteData(FStream& OutStream) = 0;
 	
-	/** Read packet data from stream */
-	virtual void ReadData(FOSCStream& OutStream) = 0;
+		/** Read packet data from stream */
+		virtual void ReadData(FStream& OutStream) = 0;
 	
-	/** Returns true if packet is message */
-	virtual bool IsMessage() = 0;
+		/** Returns true if packet is message */
+		virtual bool IsMessage() = 0;
 	
-	/** Returns true if packet is bundle */
-	virtual bool IsBundle() = 0;
+		/** Returns true if packet is bundle */
+		virtual bool IsBundle() = 0;
 
-	/** Get endpoint IP address responsible for creation/forwarding of packet */
-	virtual const FString& GetIPAddress() const;
+		UE_DEPRECATED(5.5, "Packets now have address and port set using endpoints.")
+		virtual const FString& GetIPAddress() const = 0;
 
-	/** Get endpoint address responsible for creation/forwarding of packet */
-	virtual uint16 GetPort() const;
+		UE_DEPRECATED(5.5, "Packets now have port and port set using endpoints.")
+		virtual uint16 GetPort() const = 0;
 
-	/** Create an OSC packet according to the input data. */
-	static TSharedPtr<IOSCPacket> CreatePacket(const uint8* InPacketType, const FString& InAddress, uint16 InPort);
+		virtual const FIPv4Endpoint& GetIPEndpoint() const = 0;
 
-protected:
-	FString IPAddress;
-	uint16 Port;
+		/** Create an OSC packet according to the input data. */
+		static TSharedPtr<IPacket> CreatePacket(const uint8* InPacketType, const FIPv4Endpoint& InEndpoint);
+	};
+} // namespace UE::OSC
+
+// Exists for back compat.  To be deprecated
+class OSC_API IOSCPacket : public UE::OSC::IPacket
+{
+public:
+	UE_DEPRECATED(5.5, "Use UE::OSC::IPacket instead")
+	IOSCPacket() = default;
+
+	virtual ~IOSCPacket() = default;
 };
diff --git a/Engine/Plugins/Runtime/OSC/Source/OSC/Public/OSCServer.h b/Engine/Plugins/Runtime/OSC/Source/OSC/Public/OSCServer.h
--- a/Engine/Plugins/Runtime/OSC/Source/OSC/Public/OSCServer.h
+++ b/Engine/Plugins/Runtime/OSC/Source/OSC/Public/OSCServer.h
@@ -159,7 +159,7 @@ public:
 	void ClearPackets();
 
 	/** Enqueues packet to be processed */
-	void EnqueuePacket(TSharedPtr<IOSCPacket> InPacket);
+	void EnqueuePacket(TSharedPtr<UE::OSC::IPacket> InPacket);
 
 	/** Callback for when packet is received by server */
 	void PumpPacketQueue(const TSet<uint32>* InAllowlistedClients);
@@ -178,7 +178,7 @@ private:
 	TUniquePtr<IOSCServerProxy> ServerProxy;
 
 	/** Queue stores incoming OSC packet requests to process on the game thread. */
-	TQueue<TSharedPtr<IOSCPacket>> OSCPackets;
+	TQueue<TSharedPtr<UE::OSC::IPacket>> OSCPackets;
 
 	/** Address pattern hash to check against when dispatching incoming messages */
 	TMap<FOSCAddress, FOSCDispatchMessageEvent> AddressPatterns;
diff --git a/Engine/Plugins/Runtime/OSC/Source/OSC/Public/OSCStream.h b/Engine/Plugins/Runtime/OSC/Source/OSC/Public/OSCStream.h
--- a/Engine/Plugins/Runtime/OSC/Source/OSC/Public/OSCStream.h
+++ b/Engine/Plugins/Runtime/OSC/Source/OSC/Public/OSCStream.h
@@ -1,98 +1,165 @@
 // Copyright Epic Games, Inc. All Rights Reserved.
 #pragma once
 
-#include "CoreMinimal.h"
+#include "Containers/Array.h"
+#include "UObject/NoExportTypes.h"
 
-class OSC_API FOSCStream
-{
-public:
-	FOSCStream();
-	FOSCStream(const uint8* InData, int32 InSize);
 
-	/** Returns the stream buffer data */
-	const uint8* GetData() const;
+namespace UE::OSC
+{
+	class OSC_API FStream
+	{
+	public:
+		FStream();
+		FStream(const uint8* InData, int32 InSize);
+
+		/** Returns the stream buffer data */
+		const uint8* GetData() const;
 	
-	/** Returns the stream's length */
-	int32 GetLength() const;
+		/** Returns the stream's length */
+		int32 GetLength() const;
 	
-	/** Returns true if stream has reached the end, false if not */
-	bool HasReachedEnd() const;
+		/** Returns true if stream has reached the end, false if not */
+		bool HasReachedEnd() const;
 	
-	/** Get current stream position. */
-	int32 GetPosition() const;
+		/** Get current stream position. */
+		int32 GetPosition() const;
 
-	/** Set stream position. */
-	void SetPosition(int32 InPosition);
+		/** Set stream position. */
+		void SetPosition(int32 InPosition);
 
-	/** Read Char from the stream */
-	TCHAR ReadChar();
+		/** Read Char from the stream */
+		TCHAR ReadChar();
 
-	/** Write an ansi char into the stream */
-	void WriteChar(TCHAR Char);
+		/** Write an ansi char into the stream */
+		void WriteChar(TCHAR Char);
 
-	/** Write Color into the stream */
-	void WriteColor(FColor Color);
+		/** Write Color into the stream */
+		void WriteColor(FColor Color);
 
-	/** Read Color from the stream */
-	FColor ReadColor();
+		/** Read Color from the stream */
+		FColor ReadColor();
 
-	/** Read Int32 from the stream */
-	int32 ReadInt32();
+		/** Read Int32 from the stream */
+		int32 ReadInt32();
 
-	/** Write Int32 into the stream */
-	void WriteInt32(int32 Value);
+		/** Write Int32 into the stream */
+		void WriteInt32(int32 Value);
 
-	/** Read Double from the stream */
-	double ReadDouble();
+		/** Read Double from the stream */
+		double ReadDouble();
 
-	/** Write Double into the stream */
-	void WriteDouble(uint64 Value);
+		/** Write Double into the stream */
+		void WriteDouble(uint64 Value);
 	
-	/** Read Int64 from the stream */
-	int64 ReadInt64();
+		/** Read Int64 from the stream */
+		int64 ReadInt64();
 
-	/** Write Int64 into the stream */
-	void WriteInt64(int64 Value);
+		/** Write Int64 into the stream */
+		void WriteInt64(int64 Value);
 
-	/** Read UInt64 from the stream */
-	uint64 ReadUInt64();
+		/** Read UInt64 from the stream */
+		uint64 ReadUInt64();
 
-	/** Write UInt64 into the stream */
-	void WriteUInt64(uint64 Value);
+		/** Write UInt64 into the stream */
+		void WriteUInt64(uint64 Value);
 
-	/** Read Float from the stream */
-	float ReadFloat();
+		/** Read Float from the stream */
+		float ReadFloat();
 		
-	/** Write Int64 into the stream */
-	void WriteFloat(float Value);
+		/** Write Int64 into the stream */
+		void WriteFloat(float Value);
 	
-	/** Read String from the stream */
-	FString ReadString();
+		/** Read String from the stream */
+		FString ReadString();
 
-	/** Write String into the stream */
-	void WriteString(const FString& String);
+		/** Write String into the stream */
+		void WriteString(const FString& String);
 	
-	/** Read blob from the stream */
-	TArray<uint8> ReadBlob();
+		/** Read blob from the stream */
+		TArray<uint8> ReadBlob();
+
+		/** Write blob into the stream */
+		void WriteBlob(TArray<uint8>& Blob);
+
+	private:
+		template <typename TNumericType>
+		TNumericType ReadNumeric()
+		{
+			constexpr int32 NumBytes = sizeof(TNumericType);
+			uint8 Temp[NumBytes];
+			if (Read(Temp, NumBytes) == NumBytes)
+			{
+#if PLATFORM_LITTLE_ENDIAN
+				union
+				{
+					TNumericType Number;
+					uint8 Data[sizeof(TNumericType)];
+				} Value;
+
+				for (int32 Index = 0; Index < NumBytes; ++Index)
+				{
+					Value.Data[Index] = Temp[NumBytes - Index - 1];
+				}
+
+				return Value.Number;
+#else // PLATFORM_LITTLE_ENDIAN
+				return *(TNumericType*)Temp;
+#endif // !PLATFORM_LITTLE_ENDIAN
+			}
+
+			return 0;
+		}
+
+		template <typename TNumericType>
+		void WriteNumeric(TNumericType InValue)
+		{
+			constexpr int32 NumBytes = sizeof(TNumericType);
+			uint8 Temp[NumBytes];
+
+#if PLATFORM_LITTLE_ENDIAN
+			union
+			{
+				TNumericType Number;
+				uint8 Data[sizeof(TNumericType)];
+			} Value;
+
+			Value.Number = InValue;
+
+			for (int32 Index = 0; Index < NumBytes; ++Index)
+			{
+				Temp[NumBytes - Index - 1] = Value.Data[Index];
+			}
+#else // PLATFORM_LITTLE_ENDIAN
+			return *(TNumericType*)(Temp) = InValue;
+#endif // !PLATFORM_LITTLE_ENDIAN
+
+			Write(Temp, 4);
+		}
+
+		/** Read data from stream into a provided buffer. */
+		int32 Read(uint8* OutBuffer, int32 InSize);
+	
+		/** Write data into stream from provided buffer. */
+		int32 Write(const uint8* InBuffer, int32 InSize);
 
-	/** Write blob into the stream */
-	void WriteBlob(TArray<uint8>& Blob);
+		/** Stream data. */
+		TArray<uint8> Data;
 
-private:
-	
-	/** Read data from a provided buffer. */
-	int32 Read(uint8* InBuffer, int32 InSize);
-	
-	/** Write data into buffer. */
-	int32 Write(const uint8* InBuffer, int32 InSize);
+		/** Current buffer position. */
+		int32 Position = 0;
 
-	/** Stream data. */
-	TArray<uint8> Data;
+		/** Whether stream is used to read (true) or write (false) */
+		bool bIsReadStream = false;
+	};
+} // namespace UE::OSC
 
-	/** Current buffer position. */
-	int32 Position;
+class OSC_API FOSCStream : public UE::OSC::FStream
+{
+public:
+	UE_DEPRECATED(5.5, "Use UE::OSC::FStream instead")
+	FOSCStream() = default;
 
-	/** Whether stream is used to read (true) or write (false) */
-	bool bIsReadStream;
+	UE_DEPRECATED(5.5, "Use UE::OSC::FStream instead")
+	FOSCStream(const uint8* InData, int32 InSize) : UE::OSC::FStream(InData, InSize) { }
 };
-
diff --git a/Engine/Plugins/Runtime/OSC/Source/OSC/Public/OSCTypes.h b/Engine/Plugins/Runtime/OSC/Source/OSC/Public/OSCTypes.h
--- a/Engine/Plugins/Runtime/OSC/Source/OSC/Public/OSCTypes.h
+++ b/Engine/Plugins/Runtime/OSC/Source/OSC/Public/OSCTypes.h
@@ -1,181 +1,205 @@
 // Copyright Epic Games, Inc. All Rights Reserved.
 #pragma once
 
-#include "CoreMinimal.h"
-
+#include "CoreTypes.h"
+#include "Containers/ArrayView.h"
+#include "Containers/StringView.h"
 #include "Math/Color.h"
+#include "Misc/TVariant.h"
 
 
-enum class EOSCTypeTag : uint8
+namespace UE::OSC
 {
-	OSC_INT32		= 'i',
-	OSC_FLOAT		= 'f',
-	OSC_DOUBLE		= 'd',
-	OSC_STRING		= 's',
-	OSC_BLOB		= 'b',
-	OSC_TIME		= 't',
-	OSC_INT64		= 'h',
-	OSC_CHAR		= 'c',
-	OSC_TRUE		= 'T',
-	OSC_FALSE		= 'F',
-	OSC_NIL			= 'N',
-	OSC_INFINITUM	= 'I',
-	OSC_COLOR		= 'r',
-	OSC_TERMINATE	= '\0'
-};
+	enum class EDataType : uint8
+	{
+		Blob = 'b',
+		Char = 'c',
+		Color = 'r',
+		Double = 'd',
+		False = 'F',
+		Float = 'f',
+		Infinitum = 'I',
+		Int32 = 'i',
+		Int64 = 'h',
+		Nil = 'N',
+		String = 's',
+		Terminate = '\0',
+		Time = 't',
+		True = 'T'
+	};
 
+	const OSC_API TCHAR* LexToString(EDataType DataType);
+
+	class OSC_API FOSCData
+	{
+	public:
+		FOSCData() = default;
+
+		explicit FOSCData(const TArray<uint8>& Value);
+		explicit FOSCData(TArray<uint8>&& Value);
+		explicit FOSCData(bool Value);
+		explicit FOSCData(ANSICHAR Value);
+		explicit FOSCData(FColor Value);
+		explicit FOSCData(double Value);
+		explicit FOSCData(float Value);
+		explicit FOSCData(int32 Value);
+		explicit FOSCData(int64 Value);
+		explicit FOSCData(FString Value);
+		explicit FOSCData(uint64 Value);
+
+		UE_DEPRECATED(5.5, "Use applicable explicitly typed constructor or static construction function")
+		explicit FOSCData(EDataType DataType);
+
+		static const FOSCData& Nil();
+		static const FOSCData& Infinitum();
+		static const FOSCData& Terminate();
+		static bool IsNil(const FOSCData& InType);
+
+		EDataType GetDataType() const;
+
+		FORCEINLINE bool IsBlob() const { return DataType == EDataType::Blob; }
+		FORCEINLINE bool IsBool() const { return DataType == EDataType::True || DataType == EDataType::False; }
+		FORCEINLINE bool IsChar() const { return DataType == EDataType::Char; }
+		FORCEINLINE bool IsColor() const { return DataType == EDataType::Color; }
+		FORCEINLINE bool IsDouble() const { return DataType == EDataType::Double; }
+		FORCEINLINE bool IsFloat() const { return DataType == EDataType::Float; }
+		FORCEINLINE bool IsInfinitum() const { return DataType == EDataType::Infinitum; }
+		FORCEINLINE bool IsInt32() const { return DataType == EDataType::Int32; }
+		FORCEINLINE bool IsInt64() const { return DataType == EDataType::Int64; }
+		FORCEINLINE bool IsNil() const { return DataType == EDataType::Nil; }
+		FORCEINLINE bool IsString() const { return DataType == EDataType::String; }
+		FORCEINLINE bool IsTimeTag() const { return DataType == EDataType::Time; }
+		FORCEINLINE bool IsTerminate() const { return DataType == EDataType::Terminate; }
+
+		TArray<uint8> GetBlob() const;
+		bool GetBool() const;
+		ANSICHAR GetChar() const;
+		FColor GetColor() const;
+		double GetDouble() const;
+		float GetFloat() const;
+		int32 GetInt32() const;
+		int64 GetInt64() const;
+		FString GetString() const;
+		uint64 GetTimeTag() const;
+
+		TArrayView<const uint8> GetBlobArrayView() const;
+		FStringView GetStringView() const;
+
+		using FVariant = TVariant
+		<
+			TArray<uint8>,	// Blob
+			bool,
+			ANSICHAR,
+			FColor,
+			double,
+			float,
+			int32,
+			int64,
+			FString,
+			uint64			// TimeTag
+		>;
+
+	protected:
+		EDataType DataType = EDataType::Nil;
+		FVariant Data;
+	};
+} // namespace UE::OSC
 
-class OSC_API FOSCType
+
+// Exists for back compat.  To be deprecated
+enum EOSCTypeTag
+{
+	OSC_BLOB = 'b',
+	OSC_CHAR = 'c',
+	OSC_COLOR = 'r',
+	OSC_DOUBLE = 'd',
+	OSC_FALSE = 'F',
+	OSC_FLOAT = 'f',
+	OSC_INFINITUM = 'I',
+	OSC_INT32 = 'i',
+	OSC_INT64 = 'h',
+	OSC_NIL = 'N',
+	OSC_STRING = 's',
+	OSC_TERMINATE = '\0',
+	OSC_TIME = 't',
+	OSC_TRUE = 'T'
+};
+
+class OSC_API FOSCType : public UE::OSC::FOSCData
 {
 public:
-	explicit FOSCType(int32 Value)
-		: TypeTag(EOSCTypeTag::OSC_INT32)
-		, Data(Value)
-		, Blob()
+	UE_DEPRECATED(5.5, "Use UE::OSC::FOSCData instead")
+	explicit FOSCType(const TArray<uint8>& Value)
+		: UE::OSC::FOSCData(Value)
 	{
 	}
-	explicit FOSCType(int64 Value)
-		: TypeTag(EOSCTypeTag::OSC_INT64)
-		, Data(Value)
-		, Blob()
+
+	UE_DEPRECATED(5.5, "Use UE::OSC::FOSCData instead")
+	explicit FOSCType(TArray<uint8>&& Value)
+		: UE::OSC::FOSCData(MoveTemp(Value))
 	{
 	}
-	explicit FOSCType(ANSICHAR Value)
-		: TypeTag(EOSCTypeTag::OSC_CHAR)
-		, Data(Value)
-		, Blob()
+
+	UE_DEPRECATED(5.5, "Use UE::OSC::FOSCData instead")
+	explicit FOSCType(bool Value)
+		: UE::OSC::FOSCData(Value)
 	{
 	}
-	explicit FOSCType(uint64 Value)
-		: TypeTag(EOSCTypeTag::OSC_TIME)
-		, Data(Value)
-		, Blob()
+
+	UE_DEPRECATED(5.5, "Use UE::OSC::FOSCData instead")
+	explicit FOSCType(ANSICHAR Value)
+		: UE::OSC::FOSCData(Value)
 	{
 	}
-	explicit FOSCType(float Value)
-		: TypeTag(EOSCTypeTag::OSC_FLOAT)
-		, Data(Value)
-		, Blob()
-	
+
+	UE_DEPRECATED(5.5, "Use UE::OSC::FOSCData instead")
+	explicit FOSCType(FColor Value)
+		: UE::OSC::FOSCData(Value)
 	{
 	}
+
+	UE_DEPRECATED(5.5, "Use UE::OSC::FOSCData instead")
 	explicit FOSCType(double Value)
-		: TypeTag(EOSCTypeTag::OSC_DOUBLE)
-		, Data(Value)
-		, Blob()
+		: UE::OSC::FOSCData(Value)
 	{
 	}
-	explicit FOSCType(bool Value)
-		: TypeTag(Value ? EOSCTypeTag::OSC_TRUE : EOSCTypeTag::OSC_FALSE)
-		, Data(Value)
-		, Blob()
+
+	UE_DEPRECATED(5.5, "Use UE::OSC::FOSCData instead")
+	explicit FOSCType(float Value)
+		: UE::OSC::FOSCData(Value)
 	{
 	}
-	explicit FOSCType(const FString& Value)
-		: TypeTag(EOSCTypeTag::OSC_STRING)
-		, Data(0)
-		, String(Value)
-		, Blob()
+
+	UE_DEPRECATED(5.5, "Use UE::OSC::FOSCData instead")
+	explicit FOSCType(int32 Value)
+		: UE::OSC::FOSCData(Value)
 	{
 	}
-	explicit FOSCType(const TArray<uint8>& Value)
-		: TypeTag(EOSCTypeTag::OSC_BLOB)
-		, Data(0)
-		, Blob(Value)
+
+	UE_DEPRECATED(5.5, "Use UE::OSC::FOSCData instead")
+	explicit FOSCType(int64 Value)
+		: UE::OSC::FOSCData(Value)
 	{
 	}
-	explicit FOSCType(FColor Value)
-		: TypeTag(EOSCTypeTag::OSC_COLOR)
-		, Data(0)
-		, Blob()
-		, Color(Value)
+
+	UE_DEPRECATED(5.5, "Use UE::OSC::FOSCData instead")
+	explicit FOSCType(const FString& Value)
+		: UE::OSC::FOSCData(Value)
 	{
 	}
 
-	explicit FOSCType(EOSCTypeTag TypeTag)
-		: TypeTag(TypeTag)
-		, Data(0)
-		, Blob()
+	UE_DEPRECATED(5.5, "Use UE::OSC::FOSCData instead")
+	explicit FOSCType(FString&& Value)
+		: UE::OSC::FOSCData(MoveTemp(Value))
 	{
 	}
 
-	EOSCTypeTag GetTypeTag() const { return TypeTag; }
-
-	bool IsInt32() const { return TypeTag == EOSCTypeTag::OSC_INT32; }
-	int32 GetInt32() const { return Data.Int32; }
-
-	bool IsInt64() const { return TypeTag == EOSCTypeTag::OSC_INT64; }
-	int64 GetInt64() const { return Data.Int64; }
-
-	bool IsTimeTag() const { return TypeTag == EOSCTypeTag::OSC_TIME; }
-	uint64 GetTimeTag() const { return Data.Time; }
-
-	bool IsBool() const { return TypeTag == EOSCTypeTag::OSC_TRUE || TypeTag == EOSCTypeTag::OSC_FALSE; }
-	bool GetBool() const { return Data.Bool; }
-
-	bool IsChar() const { return TypeTag == EOSCTypeTag::OSC_CHAR; }
-	ANSICHAR GetChar() const { return Data.Char; }
-
-	bool IsFloat() const { return TypeTag == EOSCTypeTag::OSC_FLOAT; }
-	float GetFloat() const { return Data.Float; }
-	
-	bool IsDouble() const { return TypeTag == EOSCTypeTag::OSC_DOUBLE; }
-	double GetDouble() const { return Data.Double; }
-
-	bool IsString() const { return TypeTag == EOSCTypeTag::OSC_STRING; }
-	FString GetString() const { return String; }
-
-	bool IsBlob() const { return TypeTag == EOSCTypeTag::OSC_BLOB; }
-	TArray<uint8> GetBlob() const { return Blob; }
-
-	bool IsColor() const { return TypeTag == EOSCTypeTag::OSC_COLOR; }
-	FColor GetColor()  const { return Color; }
-
-	bool IsNil() const { return TypeTag == EOSCTypeTag::OSC_NIL; }
-	bool IsInfinitum() const { return TypeTag == EOSCTypeTag::OSC_INFINITUM; }
-
-private:
-	union DataTypes
+	UE_DEPRECATED(5.5, "Use UE::OSC::FOSCData instead")
+	explicit FOSCType(uint64 Value)
+		: UE::OSC::FOSCData(Value)
 	{
-		explicit DataTypes(int32 Value)
-			: Int32(Value)
-		{
-		}
-		explicit DataTypes(int64 Value)
-			: Int64(Value)
-		{
-		}
-		explicit DataTypes(uint64 Value)
-			: Time(Value)
-		{
-		}
-		explicit DataTypes(ANSICHAR Value)
-			: Char(Value)
-		{
-		}
-		explicit DataTypes(float Value)
-			: Float(Value)
-		{
-		}
-		explicit DataTypes(double Value)
-			: Double(Value)
-		{
-		}
-
-		int32 Int32;
-		int64 Int64;
-		uint64 Time;
-		ANSICHAR Char;
-		float Float;
-		double Double;
-		bool Bool;
-	};
+	}
 
-	
-private:
-	EOSCTypeTag TypeTag;
-	DataTypes Data;
-	FString String;
-	TArray<uint8> Blob;
-	FColor Color;
+	UE_DEPRECATED(5.5, "Use UE::OSC::FOSCData::GetDataType() instead")
+	int32 GetTypeTag() const { return static_cast<int32>(UE::OSC::EDataType::Nil); }
 };
