Fix for Mac macro collision with OSC data type member 'Nil'

diff --git a/Engine/Plugins/Runtime/OSC/Source/OSC/Private/OSCMessagePacket.cpp b/Engine/Plugins/Runtime/OSC/Source/OSC/Private/OSCMessagePacket.cpp
--- a/Engine/Plugins/Runtime/OSC/Source/OSC/Private/OSCMessagePacket.cpp
+++ b/Engine/Plugins/Runtime/OSC/Source/OSC/Private/OSCMessagePacket.cpp
@@ -123,7 +123,7 @@ namespace UE::OSC
 				break;
 			case EDataType::True:
 			case EDataType::False:
-			case EDataType::Nil:
+			case EDataType::NilValue:
 			case EDataType::Infinitum:
 				// No values are written for these types
 				break;
@@ -177,8 +177,8 @@ namespace UE::OSC
 			case EDataType::False:
 				Arguments.Add(FOSCData(false));
 				break;
-			case EDataType::Nil:
-				Arguments.Add(FOSCData::Nil());
+			case EDataType::NilValue:
+				Arguments.Add(FOSCData::NilData());
 				break;
 			case EDataType::Infinitum:
 				Arguments.Add(FOSCData::Infinitum());
diff --git a/Engine/Plugins/Runtime/OSC/Source/OSC/Private/OSCTypes.cpp b/Engine/Plugins/Runtime/OSC/Source/OSC/Private/OSCTypes.cpp
--- a/Engine/Plugins/Runtime/OSC/Source/OSC/Private/OSCTypes.cpp
+++ b/Engine/Plugins/Runtime/OSC/Source/OSC/Private/OSCTypes.cpp
@@ -39,7 +39,7 @@ namespace UE::OSC
 				return TEXT("Int32");
 			case EDataType::Int64:
 				return TEXT("Int64");
-			case EDataType::Nil:
+			case EDataType::NilValue:
 				return TEXT("Nil");
 			case EDataType::String:
 				return TEXT("String");
@@ -133,12 +133,12 @@ namespace UE::OSC
 		Data.Set<int32>(0);
 	}
 
-	const FOSCData& FOSCData::Nil()
+	const FOSCData& FOSCData::NilData()
 	{
 		auto MakeInf = []()
 			{
 				FOSCData NilType;
-				NilType.DataType = EDataType::Nil;
+				NilType.DataType = EDataType::NilValue;
 				return NilType;
 			};
 		static const FOSCData StaticNil = MakeInf();
@@ -171,7 +171,7 @@ namespace UE::OSC
 
 	bool FOSCData::IsNil(const FOSCData& InType)
 	{
-		return InType.GetDataType() == EDataType::Nil;
+		return InType.GetDataType() == EDataType::NilValue;
 	}
 
 	EDataType FOSCData::GetDataType() const
diff --git a/Engine/Plugins/Runtime/OSC/Source/OSC/Public/OSCTypes.h b/Engine/Plugins/Runtime/OSC/Source/OSC/Public/OSCTypes.h
--- a/Engine/Plugins/Runtime/OSC/Source/OSC/Public/OSCTypes.h
+++ b/Engine/Plugins/Runtime/OSC/Source/OSC/Public/OSCTypes.h
@@ -21,7 +21,7 @@ namespace UE::OSC
 		Infinitum = 'I',
 		Int32 = 'i',
 		Int64 = 'h',
-		Nil = 'N',
+		NilValue = 'N',
 		String = 's',
 		Terminate = '\0',
 		Time = 't',
@@ -50,7 +50,7 @@ namespace UE::OSC
 		UE_DEPRECATED(5.5, "Use applicable explicitly typed constructor or static construction function")
 		explicit FOSCData(EDataType DataType);
 
-		static const FOSCData& Nil();
+		static const FOSCData& NilData();
 		static const FOSCData& Infinitum();
 		static const FOSCData& Terminate();
 		static bool IsNil(const FOSCData& InType);
@@ -66,7 +66,7 @@ namespace UE::OSC
 		FORCEINLINE bool IsInfinitum() const { return DataType == EDataType::Infinitum; }
 		FORCEINLINE bool IsInt32() const { return DataType == EDataType::Int32; }
 		FORCEINLINE bool IsInt64() const { return DataType == EDataType::Int64; }
-		FORCEINLINE bool IsNil() const { return DataType == EDataType::Nil; }
+		FORCEINLINE bool IsNil() const { return DataType == EDataType::NilValue; }
 		FORCEINLINE bool IsString() const { return DataType == EDataType::String; }
 		FORCEINLINE bool IsTimeTag() const { return DataType == EDataType::Time; }
 		FORCEINLINE bool IsTerminate() const { return DataType == EDataType::Terminate; }
@@ -100,7 +100,7 @@ namespace UE::OSC
 		>;
 
 	protected:
-		EDataType DataType = EDataType::Nil;
+		EDataType DataType = EDataType::NilValue;
 		FVariant Data;
 	};
 } // namespace UE::OSC
@@ -201,5 +201,5 @@ public:
 	}
 
 	UE_DEPRECATED(5.5, "Use UE::OSC::FOSCData::GetDataType() instead")
-	int32 GetTypeTag() const { return static_cast<int32>(UE::OSC::EDataType::Nil); }
+	int32 GetTypeTag() const { return static_cast<int32>(UE::OSC::EDataType::NilValue); }
 };
