Fix deprecation version that didn't get updated post 5.4
#rnx

diff --git a/Engine/Plugins/Runtime/OSC/Source/OSC/Private/OSCMessagePacket.h b/Engine/Plugins/Runtime/OSC/Source/OSC/Private/OSCMessagePacket.h
--- a/Engine/Plugins/Runtime/OSC/Source/OSC/Private/OSCMessagePacket.h
+++ b/Engine/Plugins/Runtime/OSC/Source/OSC/Private/OSCMessagePacket.h
@@ -20,10 +20,10 @@ namespace UE::OSC
 		/** Get endpoint IP address and port responsible for creation/forwarding of packet */
 		virtual const FIPv4Endpoint& GetIPEndpoint() const override;
 
-		UE_DEPRECATED(5.4, "Packets now have address and port set using endpoints.")
+		UE_DEPRECATED(5.5, "Packets now have address and port set using endpoints.")
 		virtual const FString& GetIPAddress() const override { static const FString None; return None; }
 
-		UE_DEPRECATED(5.4, "Packets now have port and port set using endpoints.")
+		UE_DEPRECATED(5.5, "Packets now have port and port set using endpoints.")
 		virtual uint16 GetPort() const override { return 0; };
 
 	protected:
diff --git a/Engine/Plugins/Runtime/OSC/Source/OSC/Public/OSCMessage.h b/Engine/Plugins/Runtime/OSC/Source/OSC/Public/OSCMessage.h
--- a/Engine/Plugins/Runtime/OSC/Source/OSC/Public/OSCMessage.h
+++ b/Engine/Plugins/Runtime/OSC/Source/OSC/Public/OSCMessage.h
@@ -20,7 +20,7 @@ struct OSC_API FOSCMessage
 	FOSCMessage();
 	FOSCMessage(const TSharedRef<UE::OSC::IPacket>& InPacket);
 
-	UE_DEPRECATED(5.4, "Use shared ref ctor instead")
+	UE_DEPRECATED(5.5, "Use shared ref ctor instead")
 	FOSCMessage(const TSharedPtr<UE::OSC::IPacket>& InPacket);
 
 	~FOSCMessage() = default;
@@ -28,12 +28,12 @@ struct OSC_API FOSCMessage
 	// Returns arguments, asserting if the message's packet is unset.
 	const TArray<UE::OSC::FOSCData>& GetArgumentsChecked() const;
 
-	UE_DEPRECATED(5.4, "Use shared ref setter instead")
+	UE_DEPRECATED(5.5, "Use shared ref setter instead")
 	void SetPacket(TSharedPtr<UE::OSC::IPacket>& InPacket);
 
 	void SetPacket(TSharedRef<UE::OSC::IPacket>& InPacket);
 
-	UE_DEPRECATED(5.4, "Use shared ref getter instead")
+	UE_DEPRECATED(5.5, "Use shared ref getter instead")
 	const TSharedPtr<UE::OSC::IPacket>& GetPacket() const;
 
 	const TSharedRef<UE::OSC::IPacket>& GetPacketRef() const;
