OSC Plugin Optimization/Code Standard Refactor Part 2
- Move UOSCSServer to thread-safe ticker
- Implement endpoint caching in OSCServer rather than string address/port
- Move Allowlist to support full endpoint (i.e. filter by port as well as address)
- Privitize server implementation for pumping packets that will be moved to proxy in subsequent change
- Deprecate tick in editor (servers now always tick to pump command queue)
- Move proxies to UE::OSC namespace

#rnx

diff --git a/Engine/Plugins/Runtime/OSC/Source/OSC/Private/OSCMessagePacket.h b/Engine/Plugins/Runtime/OSC/Source/OSC/Private/OSCMessagePacket.h
--- a/Engine/Plugins/Runtime/OSC/Source/OSC/Private/OSCMessagePacket.h
+++ b/Engine/Plugins/Runtime/OSC/Source/OSC/Private/OSCMessagePacket.h
@@ -23,7 +23,7 @@ namespace UE::OSC
 		UE_DEPRECATED(5.5, "Packets now have address and port set using endpoints.")
 		virtual const FString& GetIPAddress() const override { static const FString None; return None; }
 
-		UE_DEPRECATED(5.5, "Packets now have port and port set using endpoints.")
+		UE_DEPRECATED(5.5, "Packets now have address and port set using endpoints.")
 		virtual uint16 GetPort() const override { return 0; };
 
 	protected:
diff --git a/Engine/Plugins/Runtime/OSC/Source/OSC/Private/OSCServer.cpp b/Engine/Plugins/Runtime/OSC/Source/OSC/Private/OSCServer.cpp
--- a/Engine/Plugins/Runtime/OSC/Source/OSC/Private/OSCServer.cpp
+++ b/Engine/Plugins/Runtime/OSC/Source/OSC/Private/OSCServer.cpp
@@ -11,7 +11,6 @@
 
 UOSCServer::UOSCServer(const FObjectInitializer& ObjectInitializer)
 	: Super(ObjectInitializer)
-	, ServerProxy(MakeUnique<FOSCServerProxy>(*this))
 {
 }
 
@@ -31,12 +30,21 @@ void UOSCServer::Listen()
 {
 	check(ServerProxy.IsValid());
 	ServerProxy->Listen(GetName());
+
+	TickHandle = FTSTicker::GetCoreTicker().AddTicker(FTickerDelegate::CreateLambda([this](float /*Time*/)
+	{
+		PumpPacketQueue();
+		return true;
+	}));
 }
 
 bool UOSCServer::SetAddress(const FString& InReceiveIPAddress, int32 InPort)
 {
 	check(ServerProxy.IsValid());
-	return ServerProxy->SetAddress(InReceiveIPAddress, InPort);
+
+	FIPv4Address Address;
+	FIPv4Address::Parse(InReceiveIPAddress, Address);
+	return ServerProxy->SetIPEndpoint(FIPv4Endpoint(Address, InPort));
 }
 
 void UOSCServer::SetMulticastLoopback(bool bInMulticastLoopback)
@@ -48,14 +56,15 @@ void UOSCServer::SetMulticastLoopback(bool bInMulticastLoopback)
 #if WITH_EDITOR
 void UOSCServer::SetTickInEditor(bool bInTickInEditor)
 {
-	check(ServerProxy.IsValid());
-	ServerProxy->SetTickableInEditor(bInTickInEditor);
 }
 #endif // WITH_EDITOR
 
-
 void UOSCServer::Stop()
 {
+	FTSTicker::GetCoreTicker().RemoveTicker(TickHandle);
+	TickHandle.Reset();
+
+	// CDO May have this not set by initializer ctor, so have to check if valid
 	if (ServerProxy.IsValid())
 	{
 		ServerProxy->Stop();
@@ -68,54 +77,91 @@ void UOSCServer::BeginDestroy()
 	Super::BeginDestroy();
 }
 
+void UOSCServer::PostInitProperties()
+{
+	using namespace UE::OSC;
+
+	Super::PostInitProperties();
+
+	const UClass* ThisClass = UOSCServer::StaticClass();
+	check(ThisClass);
+
+	const UObject* DefaultObj = ThisClass->GetDefaultObject();
+	check(DefaultObj);
+
+	if (DefaultObj != this)
+	{
+		OSCPackets = MakeShared<FPacketQueue>();
+		ServerProxy = MakeUnique<UE::OSC::FServerProxy>(*this);
+	}
+}
+
 void UOSCServer::SetAllowlistClientsEnabled(bool bEnabled)
 {
 	check(ServerProxy.IsValid());
 	ServerProxy->SetFilterClientsByAllowList(bEnabled);
 }
 
-void UOSCServer::AddAllowlistedClient(const FString& InIPAddress)
+void UOSCServer::AddAllowlistedClient(const FString& InIPAddress, int32 IPPort)
 {
 	check(ServerProxy.IsValid());
-	ServerProxy->AddClientToAllowList(InIPAddress);
+
+	FIPv4Address NewAddress;
+	if (FIPv4Address::Parse(InIPAddress, NewAddress))
+	{
+		ServerProxy->AddClientEndpointToAllowList(FIPv4Endpoint(NewAddress, IPPort));
+	}
 }
 
-void UOSCServer::RemoveAllowlistedClient(const FString& InIPAddress)
+void UOSCServer::RemoveAllowlistedClient(const FString& InIPAddress, int32 IPPort)
 {
 	check(ServerProxy.IsValid());
-	ServerProxy->RemoveClientFromAllowList(InIPAddress);
+
+	FIPv4Address NewAddress;
+	if (FIPv4Address::Parse(InIPAddress, NewAddress))
+	{
+		ServerProxy->RemoveClientEndpointFromAllowList(FIPv4Endpoint(NewAddress, IPPort));
+	}
 }
 
 void UOSCServer::ClearAllowlistedClients()
 {
 	check(ServerProxy.IsValid());
-	ServerProxy->ClearClientAllowList();
+	ServerProxy->ClearClientEndpointAllowList();
 }
 
 FString UOSCServer::GetIpAddress(bool bIncludePort) const
 {
 	check(ServerProxy.IsValid());
 
-	FString Address = ServerProxy->GetIpAddress();
 	if (bIncludePort)
 	{
-		Address += TEXT(":");
-		Address.AppendInt(ServerProxy->GetPort());
+		return ServerProxy->GetIPEndpoint().ToString();
 	}
 
-	return Address;
+	return ServerProxy->GetIPEndpoint().Address.ToString();
 }
 
 int32 UOSCServer::GetPort() const
 {
 	check(ServerProxy.IsValid());
-	return ServerProxy->GetPort();
+	return ServerProxy->GetIPEndpoint().Port;
 }
 
-TSet<FString> UOSCServer::GetAllowlistedClients() const
+TSet<FString> UOSCServer::GetAllowlistedClients(bool bIncludePort) const
 {
 	check(ServerProxy.IsValid());
-	return ServerProxy->GetClientAllowList();
+	const TSet<FIPv4Endpoint>& Endpoints = ServerProxy->GetClientEndpointAllowList();
+
+	TSet<FString> Result;
+	Algo::Transform(Endpoints, Result, [&bIncludePort](const FIPv4Endpoint& ClientEndpoint)
+	{
+		return bIncludePort
+			? ClientEndpoint.ToString()
+			: ClientEndpoint.Address.ToString();
+	});
+
+	return Result;
 }
 
 void UOSCServer::BindEventToOnOSCAddressPatternMatchesPath(const FOSCAddress& InOSCAddressPattern, const FOSCDispatchMessageEventBP& InEvent)
@@ -158,21 +204,18 @@ void UOSCServer::UnbindAllEventsFromOnOSCAddressPatternMatching()
 TArray<FOSCAddress> UOSCServer::GetBoundOSCAddressPatterns() const
 {
 	TArray<FOSCAddress> OutAddressPatterns;
-	for (const TPair<FOSCAddress, FOSCDispatchMessageEvent>& Pair : AddressPatterns)
-	{
-		OutAddressPatterns.Add(Pair.Key);
-	}
-	return MoveTemp(OutAddressPatterns);
+	AddressPatterns.GetKeys(OutAddressPatterns);
+	return OutAddressPatterns;
 }
 
 void UOSCServer::ClearPackets()
 {
-	OSCPackets.Empty();
+	OSCPackets->Empty();
 }
 
 void UOSCServer::EnqueuePacket(TSharedPtr<UE::OSC::IPacket> InPacket)
 {
-	OSCPackets.Enqueue(InPacket);
+	OSCPackets->Enqueue(InPacket);
 }
 
 void UOSCServer::DispatchBundle(const FString& InIPAddress, uint16 InPort, const FOSCBundle& InBundle)
@@ -222,16 +265,18 @@ void UOSCServer::DispatchMessage(const FString& InIPAddress, uint16 InPort, cons
 	}
 }
 
-void UOSCServer::PumpPacketQueue(const TSet<uint32>* AllowlistedClients)
+void UOSCServer::PumpPacketQueue()
 {
 	using namespace UE::OSC;
 
+	check(IsInGameThread());
+
 	TSharedPtr<UE::OSC::IPacket> Packet;
-	while (OSCPackets.Dequeue(Packet))
+	while (OSCPackets->Dequeue(Packet))
 	{
 		FIPv4Address IPAddr;
 		const FIPv4Endpoint& Endpoint = Packet->GetIPEndpoint();
-		if (!AllowlistedClients || AllowlistedClients->Contains(Endpoint.Address.Value))
+		if (ServerProxy->CanProcessPacket(Packet.ToSharedRef()))
 		{
 			const FString StringAdr = Endpoint.Address.ToString();
 			if (Packet->IsMessage())
diff --git a/Engine/Plugins/Runtime/OSC/Source/OSC/Private/OSCServerProxy.cpp b/Engine/Plugins/Runtime/OSC/Source/OSC/Private/OSCServerProxy.cpp
--- a/Engine/Plugins/Runtime/OSC/Source/OSC/Private/OSCServerProxy.cpp
+++ b/Engine/Plugins/Runtime/OSC/Source/OSC/Private/OSCServerProxy.cpp
@@ -8,224 +8,243 @@
 #include "OSCStream.h"
 
 
-FOSCServerProxy::FOSCServerProxy(UOSCServer& InServer)
-	: Server(&InServer)
-	, Socket(nullptr)
-	, SocketReceiver(nullptr)
-	, Port(0)
-	, bMulticastLoopback(false)
-	, bFilterClientsByAllowList(false)
-#if WITH_EDITOR
-	, bTickInEditor(false)
-#endif // WITH_EDITOR
+namespace UE::OSC
 {
-}
-
-FOSCServerProxy::~FOSCServerProxy()
-{
-	Stop();
-}
-
-void FOSCServerProxy::OnPacketReceived(const FArrayReaderPtr& InData, const FIPv4Endpoint& InEndpoint)
-{
-	using namespace UE::OSC;
-
-	TSharedPtr<IPacket> Packet = IPacket::CreatePacket(InData->GetData(), InEndpoint);
-	if (!Packet.IsValid())
+	FServerProxy::FServerProxy(UOSCServer& InServer)
+		: Server(&InServer)
+		, Socket(nullptr)
+		, SocketReceiver(nullptr)
+		, bMulticastLoopback(false)
 	{
-		UE_LOG(LogOSC, Verbose, TEXT("Message received from endpoint '%s' invalid OSC packet."), *InEndpoint.ToString());
-		return;
+		ClientAllowList.Add(FIPv4Endpoint::Any);
 	}
 
-	FStream Stream = FStream(InData->GetData(), InData->Num());
-	Packet->ReadData(Stream);
-	Server->EnqueuePacket(Packet);
-}
-
-FString FOSCServerProxy::GetIpAddress() const
-{
-	return ReceiveIPAddress.ToString();
-}
-
-int32 FOSCServerProxy::GetPort() const
-{
-	return Port;
-}
-
-bool FOSCServerProxy::GetMulticastLoopback() const
-{
-	return bMulticastLoopback;
-}
-
-bool FOSCServerProxy::IsActive() const
-{
-	return SocketReceiver != nullptr;
-}
-
-void FOSCServerProxy::Listen(const FString& InServerName)
-{
-	if (IsActive())
+	FServerProxy::~FServerProxy()
 	{
-		UE_LOG(LogOSC, Error, TEXT("OSCServer currently listening: %s:%d. Failed to start new service prior to calling stop."),
-			*InServerName, *ReceiveIPAddress.ToString(), Port);
-		return;
+		Stop();
 	}
 
-	FUdpSocketBuilder Builder(*InServerName);
-	Builder.BoundToPort(Port);
-	if (ReceiveIPAddress.IsMulticastAddress())
+	bool FServerProxy::CanProcessPacket(TSharedRef<UE::OSC::IPacket> Packet) const
 	{
-		Builder.JoinedToGroup(ReceiveIPAddress);
-		if (bMulticastLoopback)
+		// 1. Check if filtering anything
+		if (ClientAllowList.Contains(FIPv4Endpoint::Any))
 		{
-			Builder.WithMulticastLoopback();
+			return true;
 		}
-	}
-	else
-	{
-		if (bMulticastLoopback)
+
+		// 2. Check for explicit endpoint
+		FIPv4Endpoint EndpointToTest = Packet->GetIPEndpoint();
+		if (ClientAllowList.Contains(EndpointToTest))
 		{
-			UE_LOG(LogOSC, Warning, TEXT("OSCServer '%s' ReceiveIPAddress provided is not a multicast address.  Not respecting MulticastLoopback boolean."),
-				*InServerName);
+			return true;
 		}
-		Builder.BoundToAddress(ReceiveIPAddress);
+
+		// 2. Check for explicit address & wildcard 'any' port endpoint
+		EndpointToTest.Port = FIPv4Endpoint::Any.Port;
+		return ClientAllowList.Contains(EndpointToTest);
 	}
 
-	Socket = Builder.Build();
-	if (Socket)
+	void FServerProxy::OnPacketReceived(const FArrayReaderPtr& InData, const FIPv4Endpoint& InEndpoint)
 	{
-		SocketReceiver = new FUdpSocketReceiver(Socket, FTimespan::FromMilliseconds(100), *(InServerName + TEXT("_ListenerThread")));
-		SocketReceiver->OnDataReceived().BindRaw(this, &FOSCServerProxy::OnPacketReceived);
-		SocketReceiver->Start();
+		using namespace UE::OSC;
 
-		UE_LOG(LogOSC, Display, TEXT("OSCServer '%s' Listening: %s:%d."), *InServerName, *ReceiveIPAddress.ToString(), Port);
+		TSharedPtr<IPacket> Packet = IPacket::CreatePacket(InData->GetData(), InEndpoint);
+		if (!Packet.IsValid())
+		{
+			UE_LOG(LogOSC, Verbose, TEXT("Message received from endpoint '%s' invalid OSC packet."), *InEndpoint.ToString());
+			return;
+		}
+
+		FStream Stream = FStream(InData->GetData(), InData->Num());
+		Packet->ReadData(Stream);
+		Server->EnqueuePacket(Packet);
 	}
-	else
+
+	FString FServerProxy::GetIpAddress() const
 	{
-		// This is expected when the server isn't available, so it's not a Warning
-		UE_LOG(LogOSC, Display, TEXT("OSCServer '%s' failed to bind to socket on %s:%d. Check that the server is available on the specified address."), *InServerName, *ReceiveIPAddress.ToString(), Port);
+		return Endpoint.Address.ToString();
 	}
-}
 
-bool FOSCServerProxy::SetAddress(const FString& InReceiveIPAddress, int32 InPort)
-{
-	if (IsActive())
+	int32 FServerProxy::GetPort() const
 	{
-		UE_LOG(LogOSC, Error, TEXT("Cannot set address while OSCServer is active."));
-		return false;
+		return Endpoint.Port;
 	}
 
-	if (!FIPv4Address::Parse(InReceiveIPAddress, ReceiveIPAddress))
+	const FIPv4Endpoint& FServerProxy::GetIPEndpoint() const
 	{
-		UE_LOG(LogOSC, Error, TEXT("Invalid ReceiveIPAddress '%s'. OSCServer ReceiveIP Address not updated."), *InReceiveIPAddress);
-		return false;
+		return Endpoint;
 	}
 
-	Port = InPort;
-	return true;
-}
+	bool FServerProxy::GetMulticastLoopback() const
+	{
+		return bMulticastLoopback;
+	}
 
-void FOSCServerProxy::SetMulticastLoopback(bool bInMulticastLoopback)
-{
-	if (bInMulticastLoopback != bMulticastLoopback && IsActive())
+	bool FServerProxy::IsActive() const
 	{
-		UE_LOG(LogOSC, Error, TEXT("Cannot update MulticastLoopback while OSCServer is active."));
-		return;
+		return SocketReceiver != nullptr;
 	}
 
-	bMulticastLoopback = bInMulticastLoopback;
-}
+	void FServerProxy::Listen(const FString& InServerName)
+	{
+		if (IsActive())
+		{
+			UE_LOG(LogOSC, Error, TEXT("OSCServer currently listening: %s:%d. Failed to start new service prior to calling stop."),
+				*InServerName, *Endpoint.Address.ToString(), Endpoint.Port);
+			return;
+		}
 
-#if WITH_EDITOR
-bool FOSCServerProxy::IsTickableInEditor() const
-{
-	return bTickInEditor;
-}
+		FUdpSocketBuilder Builder(*InServerName);
+		Builder.BoundToPort(Endpoint.Port);
+		if (Endpoint.Address.IsMulticastAddress())
+		{
+			Builder.JoinedToGroup(Endpoint.Address);
+			if (bMulticastLoopback)
+			{
+				Builder.WithMulticastLoopback();
+			}
+		}
+		else
+		{
+			if (bMulticastLoopback)
+			{
+				UE_LOG(LogOSC, Warning, TEXT("OSCServer '%s' ReceiveIPAddress provided is not a multicast address.  Not respecting MulticastLoopback boolean."),
+					*InServerName);
+			}
+			Builder.BoundToAddress(Endpoint.Address);
+		}
 
-void FOSCServerProxy::SetTickableInEditor(bool bInTickInEditor)
-{
-	bTickInEditor = bInTickInEditor;
-}
-#endif // WITH_EDITOR
+		Socket = Builder.Build();
+		if (Socket)
+		{
+			SocketReceiver = new FUdpSocketReceiver(Socket, FTimespan::FromMilliseconds(100), *(InServerName + TEXT("_ListenerThread")));
+			SocketReceiver->OnDataReceived().BindRaw(this, &FServerProxy::OnPacketReceived);
+			SocketReceiver->Start();
 
-void FOSCServerProxy::Stop()
-{
-	if (SocketReceiver)
-	{
-		delete SocketReceiver;
-		SocketReceiver = nullptr;
+			UE_LOG(LogOSC, Display, TEXT("OSCServer '%s' Listening: %s:%d."), *InServerName, *Endpoint.Address.ToString(), Endpoint.Port);
+		}
+		else
+		{
+			// This is expected when the server isn't available, so it's not a Warning
+			UE_LOG(LogOSC, Display, TEXT("OSCServer '%s' failed to bind to socket on %s:%d. Check that the server is available on the specified address."), *InServerName, *Endpoint.Address.ToString(), Endpoint.Port);
+		}
 	}
 
-	if (Socket)
+	bool FServerProxy::SetAddress(const FString& InReceiveIPAddress, int32 InPort)
 	{
-		Socket->Close();
-		ISocketSubsystem::Get(PLATFORM_SOCKETSUBSYSTEM)->DestroySocket(Socket);
-		Socket = nullptr;
-	}
+		if (IsActive())
+		{
+			UE_LOG(LogOSC, Error, TEXT("Cannot set address while OSCServer is active."));
+			return false;
+		}
 
+		FIPv4Address Address;
+		if (!FIPv4Address::Parse(InReceiveIPAddress, Address))
+		{
+			UE_LOG(LogOSC, Error, TEXT("Invalid ReceiveIPAddress '%s'. OSCServer ReceiveIP Address not updated."), *InReceiveIPAddress);
+			return false;
+		}	
 
-}
+		Endpoint = FIPv4Endpoint(Address, InPort);
+		return true;
+	}
 
-void FOSCServerProxy::AddClientToAllowList(const FString& InIPAddress)
-{
-	FIPv4Address OutAddress;
-	if (!FIPv4Address::Parse(InIPAddress, OutAddress))
+	bool FServerProxy::SetIPEndpoint(const FIPv4Endpoint& InEndpoint)
 	{
-		UE_LOG(LogOSC, Warning, TEXT("OSCServer failed to add IP Address '%s' to allow list. Address is invalid."), *InIPAddress);
-		return;
-	}
+		if (IsActive())
+		{
+			UE_LOG(LogOSC, Error, TEXT("Cannot set address while OSCServer is active."));
+			return false;
+		}
 
-	ClientAllowList.Add(OutAddress.Value);
-}
+		Endpoint = InEndpoint;
+		return true;
+	}
 
-void FOSCServerProxy::RemoveClientFromAllowList(const FString& InIPAddress)
-{
-	FIPv4Address OutAddress;
-	if (!FIPv4Address::Parse(InIPAddress, OutAddress))
+	bool FServerProxy::SetMulticastLoopback(bool bInMulticastLoopback)
 	{
-		UE_LOG(LogOSC, Warning, TEXT("OSCServer failed to remove IP Address '%s' from allow list. Address is invalid."), *InIPAddress);
-		return;
+		if (bInMulticastLoopback != bMulticastLoopback && IsActive())
+		{
+			UE_LOG(LogOSC, Error, TEXT("Cannot update MulticastLoopback while OSCServer is active."));
+			return false;
+		}
+
+		bMulticastLoopback = bInMulticastLoopback;
+		return true;
 	}
 
-	ClientAllowList.Remove(OutAddress.Value);
-}
+	void FServerProxy::Stop()
+	{
+		if (SocketReceiver)
+		{
+			delete SocketReceiver;
+			SocketReceiver = nullptr;
+		}
 
-void FOSCServerProxy::ClearClientAllowList()
-{
-	ClientAllowList.Reset();
-}
+		if (Socket)
+		{
+			Socket->Close();
+			ISocketSubsystem::Get(PLATFORM_SOCKETSUBSYSTEM)->DestroySocket(Socket);
+			Socket = nullptr;
+		}
 
-TSet<FString> FOSCServerProxy::GetClientAllowList() const
-{
-	TSet<FString> Result;
-	for (uint32 Client : ClientAllowList)
+
+	}
+
+	void FServerProxy::AddClientToAllowList(const FString& InIPAddress)
 	{
-		Result.Add(FIPv4Address(Client).ToString());
+		FIPv4Endpoint EndpointToAdd;
+		if (!FIPv4Address::Parse(InIPAddress, EndpointToAdd.Address))
+		{
+			UE_LOG(LogOSC, Warning, TEXT("OSCServer failed to add IP Address '%s' to allow list. Address is invalid."), *InIPAddress);
+			return;
+		}
+
+		ClientAllowList.Add(EndpointToAdd);
 	}
 
-	return Result;
-}
+	void FServerProxy::RemoveClientFromAllowList(const FString& InIPAddress)
+	{
+		FIPv4Endpoint EndpointToRemove;
+		if (!FIPv4Address::Parse(InIPAddress, EndpointToRemove.Address))
+		{
+			UE_LOG(LogOSC, Warning, TEXT("OSCServer failed to remove IP Address '%s' from allow list. Address is invalid."), *InIPAddress);
+			return;
+		}
 
-void FOSCServerProxy::SetFilterClientsByAllowList(bool bInEnabled)
-{
-	bFilterClientsByAllowList = bInEnabled;
-}
+		ClientAllowList.Remove(EndpointToRemove);
+	}
 
-void FOSCServerProxy::Tick(float InDeltaTime)
-{
-	check(IsInGameThread());
-	check(Server);
+	const TSet<FIPv4Endpoint>& FServerProxy::GetClientEndpointAllowList() const
+	{
+		return ClientAllowList;
+	}
 
-	Server->PumpPacketQueue(bFilterClientsByAllowList ? &ClientAllowList : nullptr);
-}
+	void FServerProxy::AddClientEndpointToAllowList(const FIPv4Endpoint& InIPEndpoint)
+	{
+		ClientAllowList.Add(InIPEndpoint);
+	}
 
-TStatId FOSCServerProxy::GetStatId() const
-{
-	RETURN_QUICK_DECLARE_CYCLE_STAT(FOSCServerProxy, STATGROUP_Tickables);
-}
+	void FServerProxy::RemoveClientEndpointFromAllowList(const FIPv4Endpoint& InIPEndpoint)
+	{
+		ClientAllowList.Add(InIPEndpoint);
+	}
 
-UWorld* FOSCServerProxy::GetTickableGameObjectWorld() const
-{
-	check(Server);
-	return Server->GetWorld();
-}
+	void FServerProxy::ClearClientEndpointAllowList()
+	{
+		ClientAllowList.Empty();
+	}
+
+	void FServerProxy::SetFilterClientsByAllowList(bool bInEnabled)
+	{
+		if (!bInEnabled)
+		{
+			ClientAllowList.Empty(1);
+			ClientAllowList.Add(FIPv4Endpoint::Any);
+		}
+		else
+		{
+			ClientAllowList.Remove(FIPv4Endpoint::Any);
+		}
+	}
+} // namespace UE::OSC
diff --git a/Engine/Plugins/Runtime/OSC/Source/OSC/Private/OSCServerProxy.h b/Engine/Plugins/Runtime/OSC/Source/OSC/Private/OSCServerProxy.h
--- a/Engine/Plugins/Runtime/OSC/Source/OSC/Private/OSCServerProxy.h
+++ b/Engine/Plugins/Runtime/OSC/Source/OSC/Private/OSCServerProxy.h
@@ -10,77 +10,69 @@
 struct FIPv4Endpoint;
 
 
-class OSC_API FOSCServerProxy : public IOSCServerProxy, public FTickableGameObject
+namespace UE::OSC
 {
-public:
-	FOSCServerProxy(UOSCServer& InServer);
-	virtual ~FOSCServerProxy();
+	class OSC_API FServerProxy : public IServerProxy
+	{
+	public:
+		FServerProxy(UOSCServer& InServer);
+		virtual ~FServerProxy();
 
-	// Begin IOSCServerProxy interface
-	bool GetMulticastLoopback() const override;
-	FString GetIpAddress() const override;
-	int32 GetPort() const override;
-	bool IsActive() const override;
+		// Begin IServerProxy interface
+		virtual bool CanProcessPacket(TSharedRef<UE::OSC::IPacket> Packet) const override;
+		virtual bool GetMulticastLoopback() const override;
 
-	void Listen(const FString& InServerName) override;
+		UE_DEPRECATED(5.5, "Use GetIPEndpoint instead")
+		virtual FString GetIpAddress() const override;
 
-	bool SetAddress(const FString& InReceiveIPAddress, int32 InPort) override;
-	void SetMulticastLoopback(bool bInMulticastLoopback) override;
+		UE_DEPRECATED(5.5, "Use GetIPEndpoint instead")
+		virtual int32 GetPort() const override;
 
-	void Stop() override;
+		virtual const FIPv4Endpoint& GetIPEndpoint() const override;
 
-	void AddClientToAllowList(const FString& InIPAddress) override;
-	void RemoveClientFromAllowList(const FString& InIPAddress) override;
-	void ClearClientAllowList() override;
-	TSet<FString> GetClientAllowList() const override;
+		virtual bool IsActive() const override;
 
-	void SetFilterClientsByAllowList(bool bInEnabled) override;
+		virtual void Listen(const FString& InServerName) override;
 
-#if WITH_EDITOR
-	void SetTickableInEditor(bool bInTickInEditor) override;
-#endif // WITH_EDITOR
-	// End IOSCServerProxy interface
+		virtual bool SetAddress(const FString& InReceiveIPAddress, int32 InPort) override;
 
-	/** Callback that receives data from a socket. */
-	void OnPacketReceived(const FArrayReaderPtr& InData, const FIPv4Endpoint& InEndpoint);
+		virtual bool SetIPEndpoint(const FIPv4Endpoint& InEndpoint) override;
+		virtual bool SetMulticastLoopback(bool bInMulticastLoopback) override;
 
-	// Begin FTickableGameObject interface
-	virtual void Tick(float InDeltaTime) override;
-	virtual bool IsTickable() const override { return IsActive(); }
-	virtual TStatId GetStatId() const override;
-	virtual UWorld* GetTickableGameObjectWorld() const override;
+		virtual void Stop() override;
 
-#if WITH_EDITOR
-	virtual bool IsTickableInEditor() const override;
-#endif // WITH_EDITOR
-	// End FTickableGameObject interface
+		virtual void AddClientToAllowList(const FString& InIPAddress) override;
+		virtual void RemoveClientFromAllowList(const FString& InIPAddress) override;
 
-	/** Parent server UObject */
-	UOSCServer* Server;
 
-	/** Socket used to listen for OSC packets. */
-	FSocket* Socket;
+		virtual void AddClientEndpointToAllowList(const FIPv4Endpoint& InIPv4Endpoint) override;
+		virtual void RemoveClientEndpointFromAllowList(const FIPv4Endpoint& InIPv4Endpoint) override;
+		virtual void ClearClientEndpointAllowList() override;
+		virtual const TSet<FIPv4Endpoint>& GetClientEndpointAllowList() const override;
 
-	/** UDP receiver. */
-	FUdpSocketReceiver* SocketReceiver;
+		virtual void SetFilterClientsByAllowList(bool bInEnabled) override;
+		// End IServerProxy interface
 
-	/** Only packets from this list of client addresses will be processed if bFilterClientsByAllowList is true. */
-	TSet<uint32> ClientAllowList;
+	private:
+		/** Callback that receives data from a socket. */
+		void OnPacketReceived(const FArrayReaderPtr& InData, const FIPv4Endpoint& InEndpoint);
 
-	/** IPAddress to listen for OSC packets on.  If unset, defaults to LocalHost */
-	FIPv4Address ReceiveIPAddress;
+		/** Parent server UObject */
+		UOSCServer* Server;
 
-	/** Port to listen for OSC packets on. */
-	int32 Port;
+		/** Socket used to listen for OSC packets. */
+		FSocket* Socket;
 
-	/** Whether or not to loopback if address provided is multicast */
-	bool bMulticastLoopback;
+		/** UDP receiver. */
+		FUdpSocketReceiver* SocketReceiver;
 
-	/** Whether or not to use client allow list */
-	bool bFilterClientsByAllowList;
+		/** Only packets from this list of client addresses will be processed if bFilterClientsByAllowList is true. */
+		TSet<FIPv4Endpoint> ClientAllowList;
 
-#if WITH_EDITOR
-	bool bTickInEditor;
-#endif // WITH_EDITOR
+		/** Endpoint to listen for OSC packets on. If set to 'Any', defaults to LocalHost */
+		FIPv4Endpoint Endpoint;
 
-};
+		/** Whether or not to loopback if address provided is multicast */
+		bool bMulticastLoopback = false;
+	};
+} // namespace UE::OSC	
diff --git a/Engine/Plugins/Runtime/OSC/Source/OSC/Public/OSCServer.h b/Engine/Plugins/Runtime/OSC/Source/OSC/Public/OSCServer.h
--- a/Engine/Plugins/Runtime/OSC/Source/OSC/Public/OSCServer.h
+++ b/Engine/Plugins/Runtime/OSC/Source/OSC/Public/OSCServer.h
@@ -1,16 +1,20 @@
 // Copyright Epic Games, Inc. All Rights Reserved.
 #pragma once
 
-#include "CoreMinimal.h"
 
 #include "Async/TaskGraphInterfaces.h"
 #include "Containers/Queue.h"
-#include "Templates/UniquePtr.h"
-#include "UObject/Object.h"
-
+#include "Containers/Set.h"
+#include "Containers/Ticker.h"
+#include "Interfaces/IPv4/IPv4Address.h"
+#include "Interfaces/IPv4/IPv4Endpoint.h"
 #include "OSCBundle.h"
 #include "OSCMessage.h"
 #include "OSCPacket.h"
+#include "Templates/Function.h"
+#include "Templates/SharedPointer.h"
+#include "Templates/UniquePtr.h"
+#include "UObject/Object.h"
 
 #include "OSCServer.generated.h"
 
@@ -29,30 +33,105 @@ DECLARE_DYNAMIC_DELEGATE_FourParams(FOSCDispatchMessageEventBP, const FOSCAddres
 DECLARE_STATS_GROUP(TEXT("OSC Commands"), STATGROUP_OSCNetworkCommands, STATCAT_Advanced);
 
 
-/** Interface for internal networking implementation.  See UOSCServer for details */
-class OSC_API IOSCServerProxy
+namespace UE::OSC
 {
-public:
-	virtual ~IOSCServerProxy() { }
-
-	virtual FString GetIpAddress() const = 0;
-	virtual int32 GetPort() const = 0;
-	virtual bool GetMulticastLoopback() const = 0;
-	virtual bool IsActive() const = 0;
-	virtual void Listen(const FString& ServerName) = 0;
-	virtual bool SetAddress(const FString& InReceiveIPAddress, int32 InPort) = 0;
-	virtual void SetMulticastLoopback(bool bInMulticastLoopback) = 0;
+	// Forward Declarations
+	class IServerProxy;
+
+	/** Interface for internal networking implementation.  See UOSCServer for details */
+	class OSC_API IServerProxy
+	{
+	public:
+
+		virtual ~IServerProxy() { }
+
+		// Returns whether or not packet can be processed, i.e. is valid and allowlisted.
+		virtual bool CanProcessPacket(TSharedRef<UE::OSC::IPacket> Packet) const = 0;
+
+		UE_DEPRECATED(5.5, "Use GetIPEndpoint instead")
+		virtual FString GetIpAddress() const { return { }; }
+
+		virtual const FIPv4Endpoint& GetIPEndpoint() const = 0;
+
+		UE_DEPRECATED(5.5, "Use GetIPEndpoint instead")
+		virtual int32 GetPort() const { return GetIPEndpoint().Port; }
+
+		// Returns whether or not loopback is enabled
+		virtual bool GetMulticastLoopback() const = 0;
+
+		// Returns whether or not the server is currently active (listening)
+		virtual bool IsActive() const = 0;
+
+		// Starts the server, causing it to actively listen and dispatch OSC messages
+		virtual void Listen(const FString& ServerName) = 0;
+
+		UE_DEPRECATED(5.5, "Use SetIPEndpoint instead")
+		virtual bool SetAddress(const FString& InReceiveIPAddress, int32 InPort)
+		{
+			FIPv4Endpoint Endpoint;
+			FIPv4Address::Parse(InReceiveIPAddress, Endpoint.Address);
+			Endpoint.Port = InPort;
+			return SetIPEndpoint(Endpoint);
+		}
+
+		// Sets the current server's endpoint.  Ignores request and returns false if server
+		// is currently active.
+		virtual bool SetIPEndpoint(const FIPv4Endpoint& InEndpoint) = 0;
+
+		// Sets whether or not loopback is enabled.  Returns false and request is ignored
+		// if server is currently active.
+		virtual bool SetMulticastLoopback(bool bInMulticastLoopback) = 0;
 #if WITH_EDITOR
-	virtual void SetTickableInEditor(bool bInTickInEditor) = 0;
+		UE_DEPRECATED(5.5, "All server proxies are now ticked internally by handle in all editor contexts.")
+		virtual void SetTickableInEditor(bool bInTickInEditor) { };
 #endif // WITH_EDITOR
-	virtual void Stop() = 0;
-	virtual void AddClientToAllowList(const FString& InIPAddress) = 0;
-	virtual void RemoveClientFromAllowList(const FString& IPAddress) = 0;
-	virtual void ClearClientAllowList() = 0;
-	virtual TSet<FString> GetClientAllowList() const = 0;
-	virtual void SetFilterClientsByAllowList(bool bEnabled) = 0;
+
+		// Stops the server
+		virtual void Stop() = 0;
+
+		UE_DEPRECATED(5.5, "AllowList is now managed as IPv4Endpoints. Use API that works with endpoint struct directly")
+		virtual void AddClientToAllowList(const FString& InIPAddress) { }
+
+		UE_DEPRECATED(5.5, "AllowList is now managed as IPv4Endpoints. Use API that works with endpoint struct directly")
+		virtual void RemoveClientFromAllowList(const FString& IPAddress) { }
+
+		UE_DEPRECATED(5.5, "AllowList is now managed as IPv4Endpoints. Use API that works with endpoint struct directly")
+		virtual void ClearClientAllowList() { }
+
+		UE_DEPRECATED(5.5, "AllowList is now managed as IPv4Endpoints. Use API that works with endpoint struct directly")
+		virtual TSet<FString> GetClientAllowList() const { return { }; }
+
+		// Adds the given endpoint to the client allow list
+		virtual void AddClientEndpointToAllowList(const FIPv4Endpoint& InIPv4Endpoint) = 0;
+
+		// Removes the given endpoint from the client allow list
+		virtual void RemoveClientEndpointFromAllowList(const FIPv4Endpoint& InIPv4Endpoint) = 0;
+
+		// Empties the given address to the client allow list
+		virtual void ClearClientEndpointAllowList() = 0;
+
+		// Returns Client Address allow list if it is enabled, else returns nullptr
+		virtual const TSet<FIPv4Endpoint>& GetClientEndpointAllowList() const = 0;
+
+		// Sets whether or not allow list is active.  If disabled, entries are not cleared
+		// however GetClientAddressAllowList will no longer return a valid reference to the
+		// given set of allow list members.
+		virtual void SetFilterClientsByAllowList(bool bEnabled) = 0;
+	};
+} // namespace UE::OSC
+
+
+// For backward compat.  To be deprecated
+class OSC_API IOSCServerProxy : public UE::OSC::IServerProxy
+{
+public:
+	UE_DEPRECATED(5.5, "Use UE::OSC::IServerProxy instead")
+	IOSCServerProxy() = default;
+
+	virtual ~IOSCServerProxy() = default;
 };
 
+
 UCLASS(BlueprintType)
 class OSC_API UOSCServer : public UObject
 {
@@ -105,11 +184,11 @@ public:
 
 	/** Adds client to allowlist of clients to listen for. */
 	UFUNCTION(BlueprintCallable, Category = "Audio|OSC")
-	void AddAllowlistedClient(const FString& IPAddress);
+	void AddAllowlistedClient(const FString& IPAddress, int32 IPPort = 0);
 
 	/** Removes allowlisted client to listen for. */
 	UFUNCTION(BlueprintCallable, Category = "Audio|OSC")
-	void RemoveAllowlistedClient(const FString& IPAddress);
+	void RemoveAllowlistedClient(const FString& IPAddress, int32 IPPort = 0);
 
 	/** Clears client allowlist to listen for. */
 	UFUNCTION(BlueprintCallable, Category = "Audio|OSC")
@@ -123,9 +202,9 @@ public:
 	UFUNCTION(BlueprintCallable, Category = "Audio|OSC")
 	int32 GetPort() const;
 
-	/** Returns set of allowlisted clients. */
+	/** Returns set of allowlisted endpoint clients as strings with (optional) port included. */
 	UFUNCTION(BlueprintCallable, Category = "Audio|OSC")
-	TSet<FString> GetAllowlistedClients() const;
+	TSet<FString> GetAllowlistedClients(bool bIncludePort = false) const;
 
 	/** Adds event to dispatch when OSCAddressPattern is matched. */
 	UFUNCTION(BlueprintCallable, Category = "Audio|OSC")
@@ -151,23 +230,26 @@ public:
 	/** Set whether server instance can be ticked in-editor (editor only and available to blueprint
 	  * for use in editor utility scripts/script actions).
 	  */
-	UFUNCTION(BlueprintCallable, Category = "Audio|OSC")
+	UFUNCTION(BlueprintCallable, Category = "Audio|OSC", meta = (Deprecated = "5.5", DeprecationMessage = "Servers are now implemented as dispatchers, which are pumped on an async task upon calling 'Listen' at a provided rate (i.e. no longer 'ticked' on the game thread). See 'SetUpdateRate'"))
 	void SetTickInEditor(bool bInTickInEditor);
 #endif // WITH_EDITOR
 
-	/** Clears all packets pending processing */
+	UE_DEPRECATED(5.5, "Clearing packets directly is not thread safe and no longer supported.")
 	void ClearPackets();
 
-	/** Enqueues packet to be processed */
 	void EnqueuePacket(TSharedPtr<UE::OSC::IPacket> InPacket);
 
-	/** Callback for when packet is received by server */
-	void PumpPacketQueue(const TSet<uint32>* InAllowlistedClients);
+	UE_DEPRECATED(5.5, "Pumping packets is now handled privately")
+	void PumpPacketQueue(const TSet<uint32>* InAllowlistedClients) { }
 
 protected:
-	void BeginDestroy() override;
+	virtual void BeginDestroy() override;
+	virtual void PostInitProperties() override;
 
 private:
+	using FPacketQueue = TQueue<TSharedPtr<UE::OSC::IPacket>>;
+	void PumpPacketQueue();
+
 	/** Dispatches provided bundle received */
 	void DispatchBundle(const FString& InIPAddress, uint16 InPort, const FOSCBundle& InBundle);
 
@@ -175,11 +257,13 @@ private:
 	void DispatchMessage(const FString& InIPAddress, uint16 InPort, const FOSCMessage& InMessage);
 
 	/** Pointer to internal implementation of server proxy */
-	TUniquePtr<IOSCServerProxy> ServerProxy;
+	TUniquePtr<UE::OSC::IServerProxy> ServerProxy;
 
 	/** Queue stores incoming OSC packet requests to process on the game thread. */
-	TQueue<TSharedPtr<UE::OSC::IPacket>> OSCPackets;
+	TSharedPtr<FPacketQueue> OSCPackets;
 
 	/** Address pattern hash to check against when dispatching incoming messages */
 	TMap<FOSCAddress, FOSCDispatchMessageEvent> AddressPatterns;
+
+	FTSTicker::FDelegateHandle TickHandle;
 };
