- Implement feedback from prior OSC refactor changes
- Fix missed couple calls to templatized WriteNumeric in OSC stream

#rnx

diff --git a/Engine/Plugins/Runtime/OSC/Source/OSC/Private/OSCAddress.cpp b/Engine/Plugins/Runtime/OSC/Source/OSC/Private/OSCAddress.cpp
--- a/Engine/Plugins/Runtime/OSC/Source/OSC/Private/OSCAddress.cpp
+++ b/Engine/Plugins/Runtime/OSC/Source/OSC/Private/OSCAddress.cpp
@@ -62,47 +62,19 @@ bool FOSCAddress::IsValidPath() const
 	return bIsValidPath;
 }
 
-bool FOSCAddress::PushContainer(const FString& Container)
-{
-	return PushContainers({ Container });
-}
-
-bool FOSCAddress::PushContainer(FString&& Container)
+bool FOSCAddress::PushContainer(FString Container)
 {
 	return PushContainers({ MoveTemp(Container) });
 }
 
-bool FOSCAddress::PushContainers(const TArray<FString>& InContainers)
-{
-	TArray<FString> NewContainers = InContainers;
-	const bool bPushed = PushContainersInternal(MoveTemp(NewContainers));
-	if (bPushed)
-	{
-		CacheAggregates();
-	}
-
-	return bPushed;
-}
-
-bool FOSCAddress::PushContainers(TArray<FString>&& InContainers)
-{
-	const bool bPushed = PushContainersInternal(MoveTemp(InContainers));
-	if (bPushed)
-	{
-		CacheAggregates();
-	}
-
-	return bPushed;
-}
-
-bool FOSCAddress::PushContainersInternal(TArray<FString>&& InContainers)
+bool FOSCAddress::PushContainers(TArray<FString> NewContainers)
 {
-	if (InContainers.IsEmpty())
+	if (NewContainers.IsEmpty())
 	{
 		return false;
 	}
 
-	for (const FString& Container : InContainers)
+	for (const FString& Container : NewContainers)
 	{
 		if (Container.Contains(UE::OSC::PathSeparator))
 		{
@@ -112,68 +84,58 @@ bool FOSCAddress::PushContainersInternal(TArray<FString>&& InContainers)
 		}
 	}
 
-	Containers.Append(MoveTemp(InContainers));
-	for (const FString& Container : InContainers)
-	{
-		Containers.Push(Container);
-	}
-
+	Containers.Append(MoveTemp(NewContainers));
+	CacheAggregates();
 	return true;
 }
 
-FString FOSCAddress::PopContainer(bool* bOutDidPop)
+FString FOSCAddress::PopContainer()
 {
 	FString Popped;
-	if (Containers.Num() > 0)
-	{
-		Popped = Containers.Pop(EAllowShrinking::No);
-		Hash = GetTypeHash(GetFullPath());
+	PopContainer(Popped);
+	return Popped;
+}
 
-		if (bOutDidPop)
-		{
-			*bOutDidPop = true;
-		}
-	}
-	else
+bool FOSCAddress::PopContainer(FString& OutContainer)
+{
+	OutContainer = { };
+
+	FString Popped;
+	if (Containers.IsEmpty())
 	{
-		if (bOutDidPop)
-		{
-			*bOutDidPop = false;
-		}
+		return false;
 	}
 
-	return Popped;
+	Popped = Containers.Pop(EAllowShrinking::No);
+	Hash = GetTypeHash(GetFullPath());
+	return true;
 }
 
-TArray<FString> FOSCAddress::PopContainers(int32 InNumContainers, bool* bOutDidPop)
+TArray<FString> FOSCAddress::PopContainers(int32 NumToPop)
 {
 	TArray<FString> Popped;
-	if (InNumContainers <= 0 || Containers.Num() == 0)
+	PopContainers(NumToPop, Popped);
+	return Popped;
+}
+
+bool FOSCAddress::PopContainers(int32 NumToPop, TArray<FString>& OutContainers)
+{
+	OutContainers = { };
+
+	if (NumToPop <= 0 || Containers.IsEmpty())
 	{
-		if (bOutDidPop)
-		{
-			*bOutDidPop = false;
-		}
-		return Popped;
+		return false;
 	}
 
 	int32 Removed = 0;
-	for (int32 i = Containers.Num() - 1; i >= 0; --i)
+	for (int32 Index = Containers.Num() - 1; Index >= 0 && Removed <= NumToPop; --Index)
 	{
-		if (Removed > InNumContainers)
-		{
-			break;
-		}
 		++Removed;
-		Popped.Add(Containers.Pop(EAllowShrinking::No));
+		OutContainers.Add(Containers.Pop(EAllowShrinking::No));
 	}
 
 	Hash = GetTypeHash(GetFullPath());
-	if (bOutDidPop)
-	{
-		*bOutDidPop = true;
-	}
-	return Popped;
+	return true;
 }
 
 bool FOSCAddress::RemoveContainers(int32 InIndex, int32 InCount)
@@ -201,9 +163,9 @@ const FString& FOSCAddress::GetMethod() const
 	return Method;
 }
 
-int32 FOSCAddress::GetNumericPrefix(bool* bIsNumeric) const
+bool FOSCAddress::GetNumericPrefix(int32& OutPreflix) const
 {
-	int32 Value = 0;
+	OutPreflix = 0;
 	int32 Dec = 1;
 	const TArray<TCHAR>& CharArray = Method.GetCharArray();
 	for (
@@ -211,20 +173,16 @@ int32 FOSCAddress::GetNumericPrefix(bool* bIsNumeric) const
 		Index < CharArray.Num() && FChar::IsDigit(CharArray[Index]);
 		++Index)
 	{
-		Value = (Value * Dec) + FChar::ConvertCharDigitToInt(CharArray[Index]);
+		OutPreflix = (OutPreflix * Dec) + FChar::ConvertCharDigitToInt(CharArray[Index]);
 		Dec *= 10;
 	}
 
-	if (bIsNumeric)
-	{
-		*bIsNumeric = Dec == 1;
-	}
-	return Value;
+	return Dec == 1;
 }
 
-int32 FOSCAddress::GetNumericSuffix(bool* bIsNumeric) const
+bool FOSCAddress::GetNumericSuffix(int32& OutSuffix) const
 {
-	int32 Value = 0;
+	OutSuffix = INDEX_NONE;
 	int32 Dec = 1;
 	const TArray<TCHAR>& CharArray = Method.GetCharArray();
 	for (
@@ -232,71 +190,47 @@ int32 FOSCAddress::GetNumericSuffix(bool* bIsNumeric) const
 		Index >= 0 && FChar::IsDigit(CharArray[Index]);
 		--Index)
 	{
-		Value += FChar::ConvertCharDigitToInt(CharArray[Index]) * Dec;
+		OutSuffix += FChar::ConvertCharDigitToInt(CharArray[Index]) * Dec;
 		Dec *= 10;
 	}
 
-	if (bIsNumeric)
-	{
-		*bIsNumeric = Dec == 1;
-	}
-	return Value;
-}
-
-bool FOSCAddress::Set(const TArray<FString>& InContainers, const FString& InMethod)
-{
-	TArray<FString> NewContainers = InContainers;
-	FString NewMethod = InMethod;
-	return Set(MoveTemp(NewContainers), MoveTemp(NewMethod));
+	return Dec == 1;
 }
 
-bool FOSCAddress::Set(TArray<FString>&& InContainers, FString&& InMethod)
+bool FOSCAddress::Set(TArray<FString> NewContainers, FString NewMethod)
 {
 	Containers = { };
-	Method = { };
 
-	const bool bPushContainers = PushContainersInternal(MoveTemp(InContainers));
-	if (bPushContainers)
+	if (PushContainers(MoveTemp(NewContainers)))
 	{
-		const bool bSetMethod = SetMethod(MoveTemp(InMethod)); // Calls aggregate internally, so no need to call again on success
-		if (bSetMethod)
+		if (SetMethod(MoveTemp(NewMethod))) // Calls aggregate internally, so no need to call again on success
 		{
 			return true;
 		}
-		else
-		{
-			Containers = { };
-		}
 	}
 
+	Method = { };
 	CacheAggregates();
 	return false;
 }
 
-bool FOSCAddress::SetMethod(const FString& InMethod)
+bool FOSCAddress::SetMethod(FString NewMethod)
 {
-	FString NewMethod = InMethod;
-	return SetMethod(MoveTemp(NewMethod));
-}
-
-bool FOSCAddress::SetMethod(FString&& InMethod)
-{
-	if (InMethod.IsEmpty())
+	if (NewMethod.IsEmpty())
 	{
 		UE_LOG(LogOSC, Warning, TEXT("Failed to set OSCAddress method. "
 			"'InMethod' cannot be empty string."));
 		return false;
 	}
 
-	if (InMethod.Contains(UE::OSC::PathSeparator))
+	if (NewMethod.Contains(UE::OSC::PathSeparator))
 	{
 		UE_LOG(LogOSC, Warning, TEXT("Failed to set OSCAddress method. "
 			"Cannot contain OSC path separator '%s'."), *UE::OSC::PathSeparator);
 		return false;
 	}
 
-	Method = MoveTemp(InMethod);
-
+	Method = MoveTemp(NewMethod);
 	CacheAggregates();
 	return true;
 }
diff --git a/Engine/Plugins/Runtime/OSC/Source/OSC/Private/OSCBundle.cpp b/Engine/Plugins/Runtime/OSC/Source/OSC/Private/OSCBundle.cpp
--- a/Engine/Plugins/Runtime/OSC/Source/OSC/Private/OSCBundle.cpp
+++ b/Engine/Plugins/Runtime/OSC/Source/OSC/Private/OSCBundle.cpp
@@ -5,7 +5,7 @@
 
 
 FOSCBundle::FOSCBundle()
-	: Packet(MakeShareable(new UE::OSC::FBundlePacket()))
+	: Packet(MakeShared<UE::OSC::FBundlePacket>())
 {
 }
 
diff --git a/Engine/Plugins/Runtime/OSC/Source/OSC/Private/OSCBundlePacket.cpp b/Engine/Plugins/Runtime/OSC/Source/OSC/Private/OSCBundlePacket.cpp
--- a/Engine/Plugins/Runtime/OSC/Source/OSC/Private/OSCBundlePacket.cpp
+++ b/Engine/Plugins/Runtime/OSC/Source/OSC/Private/OSCBundlePacket.cpp
@@ -8,12 +8,7 @@
 
 namespace UE::OSC
 {
-	FBundlePacket::FBundlePacket()
-		: TimeTag(0)
-	{
-	}
-
-	FBundlePacket::FBundlePacket(const FIPv4Endpoint& InEndpoint)
+	FBundlePacket::FBundlePacket(FIPv4Endpoint InEndpoint)
 		: FPacketBase(InEndpoint)
 		, TimeTag(0)
 	{
diff --git a/Engine/Plugins/Runtime/OSC/Source/OSC/Private/OSCBundlePacket.h b/Engine/Plugins/Runtime/OSC/Source/OSC/Private/OSCBundlePacket.h
--- a/Engine/Plugins/Runtime/OSC/Source/OSC/Private/OSCBundlePacket.h
+++ b/Engine/Plugins/Runtime/OSC/Source/OSC/Private/OSCBundlePacket.h
@@ -11,8 +11,7 @@ namespace UE::OSC
 	class FBundlePacket : public FPacketBase
 	{
 	public:
-		FBundlePacket();
-		FBundlePacket(const FIPv4Endpoint& InEndpoint);
+		FBundlePacket(FIPv4Endpoint InEndpoint = FIPv4Endpoint::Any);
 
 		virtual ~FBundlePacket() = default;
 
diff --git a/Engine/Plugins/Runtime/OSC/Source/OSC/Private/OSCMessage.cpp b/Engine/Plugins/Runtime/OSC/Source/OSC/Private/OSCMessage.cpp
--- a/Engine/Plugins/Runtime/OSC/Source/OSC/Private/OSCMessage.cpp
+++ b/Engine/Plugins/Runtime/OSC/Source/OSC/Private/OSCMessage.cpp
@@ -60,7 +60,7 @@ bool FOSCMessage::SetAddress(const FOSCAddress& InAddress)
 		return false;
 	}
 
-	StaticCastSharedRef<UE::OSC::FMessagePacket>(Packet)->SetAddress(InAddress);
+	StaticCastSharedRef<FMessagePacket>(Packet)->SetAddress(InAddress);
 	return true;
 }
 
diff --git a/Engine/Plugins/Runtime/OSC/Source/OSC/Private/OSCMessagePacket.cpp b/Engine/Plugins/Runtime/OSC/Source/OSC/Private/OSCMessagePacket.cpp
--- a/Engine/Plugins/Runtime/OSC/Source/OSC/Private/OSCMessagePacket.cpp
+++ b/Engine/Plugins/Runtime/OSC/Source/OSC/Private/OSCMessagePacket.cpp
@@ -8,7 +8,7 @@
 
 namespace UE::OSC
 {
-	FPacketBase::FPacketBase(const FIPv4Endpoint& InEndpoint)
+	FPacketBase::FPacketBase(FIPv4Endpoint InEndpoint)
 		: IPEndpoint(InEndpoint)
 	{
 	}
diff --git a/Engine/Plugins/Runtime/OSC/Source/OSC/Private/OSCMessagePacket.h b/Engine/Plugins/Runtime/OSC/Source/OSC/Private/OSCMessagePacket.h
--- a/Engine/Plugins/Runtime/OSC/Source/OSC/Private/OSCMessagePacket.h
+++ b/Engine/Plugins/Runtime/OSC/Source/OSC/Private/OSCMessagePacket.h
@@ -12,8 +12,7 @@ namespace UE::OSC
 	class FPacketBase : public IPacket
 	{
 	public:
-		FPacketBase() = default;
-		FPacketBase(const FIPv4Endpoint& InEndpoint);
+		FPacketBase(FIPv4Endpoint InEndpoint = FIPv4Endpoint::Any);
 
 		virtual ~FPacketBase() = default;
 
diff --git a/Engine/Plugins/Runtime/OSC/Source/OSC/Private/OSCStream.cpp b/Engine/Plugins/Runtime/OSC/Source/OSC/Private/OSCStream.cpp
--- a/Engine/Plugins/Runtime/OSC/Source/OSC/Private/OSCStream.cpp
+++ b/Engine/Plugins/Runtime/OSC/Source/OSC/Private/OSCStream.cpp
@@ -100,28 +100,8 @@ namespace UE::OSC
 
 	void FStream::WriteDouble(uint64 Value)
 	{
-		uint8 Temp[8];
-
-	#ifdef PLATFORM_LITTLE_ENDIAN
-		union {
-			double i;
-			uint8 c[8];
-		} u;
-
-		u.i = Value;
-
-		Temp[7] = u.c[0];
-		Temp[6] = u.c[1];
-		Temp[5] = u.c[2];
-		Temp[4] = u.c[3];
-		Temp[3] = u.c[4];
-		Temp[2] = u.c[5];
-		Temp[1] = u.c[6];
-		Temp[0] = u.c[7];
-	#else // PLATFORM_LITTLE_ENDIAN
-		*(double*)(Temp) = Value;
-	#endif // !PLATFORM_LITTLE_ENDIAN
-		Write(Temp, 8);
+		WriteNumeric<double>(Value);
+
 	}
 
 	int64 FStream::ReadInt64()
@@ -131,28 +111,7 @@ namespace UE::OSC
 
 	void FStream::WriteInt64(int64 Value)
 	{
-		uint8 Temp[8];
-
-	#ifdef PLATFORM_LITTLE_ENDIAN
-		union {
-			int64 i;
-			uint8 c[8];
-		} u;
-
-		u.i = Value;
-
-		Temp[7] = u.c[0];
-		Temp[6] = u.c[1];
-		Temp[5] = u.c[2];
-		Temp[4] = u.c[3];
-		Temp[3] = u.c[4];
-		Temp[2] = u.c[5];
-		Temp[1] = u.c[6];
-		Temp[0] = u.c[7];
-	#else // PLATFORM_LITTLE_ENDIAN
-		*(int64*)(Temp) = Value;
-	#endif // !PLATFORM_LITTLE_ENDIAN
-		Write(Temp, 8);
+		WriteNumeric<int64>(Value);
 	}
 
 	uint64 FStream::ReadUInt64()
diff --git a/Engine/Plugins/Runtime/OSC/Source/OSC/Public/OSCAddress.h b/Engine/Plugins/Runtime/OSC/Source/OSC/Public/OSCAddress.h
--- a/Engine/Plugins/Runtime/OSC/Source/OSC/Public/OSCAddress.h
+++ b/Engine/Plugins/Runtime/OSC/Source/OSC/Public/OSCAddress.h
@@ -50,22 +50,22 @@ public:
 	bool Matches(const FOSCAddress& InAddress) const;
 
 	/** Pushes container onto end of address' ordered array of containers */
-	bool PushContainer(const FString& InContainer);
-
-	/** Pushes container onto end of address' ordered array of containers */
-	bool PushContainer(FString&& InContainer);
+	bool PushContainer(FString Container);
 
 	/** Pushes containers onto end of address' ordered array of containers */
-	bool PushContainers(const TArray<FString>& InContainers);
+	bool PushContainers(TArray<FString> NewContainers);
 
-	/** Pushes containers onto end of address' ordered array of containers */
-	bool PushContainers(TArray<FString>&& InContainers);
+	/** Pops container from ordered array of containers */
+	FString PopContainer();
 
 	/** Pops container from ordered array of containers */
-	FString PopContainer(bool* bDidPop = nullptr);
+	bool PopContainer(FString& OutContainer);
+
+	/** Pops containers off end of address' ordered array of containers */
+	TArray<FString> PopContainers(int32 NumToPop);
 
 	/** Pops containers off end of address' ordered array of containers */
-	TArray<FString> PopContainers(int32 InNumContainers, bool* bDidPop = nullptr);
+	bool PopContainers(int32 NumToPop, TArray<FString>& OutContainers);
 
 	/** Removes containers from container array at index until count */
 	bool RemoveContainers(int32 InIndex, int32 InCount);
@@ -77,22 +77,16 @@ public:
 	const FString& GetMethod() const;
 
 	/** Returns numeric prefix of OSC method if is numeric */
-	int32 GetNumericPrefix(bool* bIsNumeric = nullptr) const;
+	bool GetNumericPrefix(int32& OutPrefix) const;
 
 	/** Returns numeric suffix of OSC method if is numeric */
-	int32 GetNumericSuffix(bool* bIsNumeric = nullptr) const;
-
-	/** Sets the address to the given array of containers and provided method */
-	bool Set(TArray<FString>&& InContainers, FString&& InMethod);
+	bool GetNumericSuffix(int32& OutSuffix) const;
 
 	/** Sets the address to the given array of containers and provided method */
-	bool Set(const TArray<FString>& InContainers, const FString& InMethod);
-
-	/** Sets the method name of address */
-	bool SetMethod(const FString& InMethod);
+	bool Set(TArray<FString> NewContainers, FString NewMethod);
 
 	/** Sets the method name of address */
-	bool SetMethod(FString&& InMethod);
+	bool SetMethod(FString NewMethod);
 
 	/** Returns container path of OSC address in the form '/Container1/Container2' */
 	FString GetContainerPath() const;
