Extract OSC packet dispatching from broadcast part 1
- Move server proxy to static create
- Move server queue to spsc implementation
- Pipe descriptor to IServerProxy from socket
- Misc clean-up

#rnx

diff --git a/Engine/Plugins/Runtime/OSC/Source/OSC/Private/OSCServer.cpp b/Engine/Plugins/Runtime/OSC/Source/OSC/Private/OSCServer.cpp
--- a/Engine/Plugins/Runtime/OSC/Source/OSC/Private/OSCServer.cpp
+++ b/Engine/Plugins/Runtime/OSC/Source/OSC/Private/OSCServer.cpp
@@ -9,6 +9,26 @@
 #include "OSCServerProxy.h"
 
 
+namespace UE::OSC
+{
+	TSharedPtr<IServerProxy> IServerProxy::Create(UOSCServer* Parent)
+	{
+		return MakeShared<OSC::FServerProxy>(*Parent);
+	}
+
+	bool IServerProxy::SetAddress(const FString& InReceiveIPAddress, int32 InPort)
+	{
+		FIPv4Address Address;
+		if (FIPv4Address::Parse(InReceiveIPAddress, Address))
+		{
+			return SetIPEndpoint(FIPv4Endpoint(Address, InPort));
+		}
+
+		UE_LOG(LogOSC, Error, TEXT("Invalid ReceiveIPAddress '%s'. OSCServer ReceiveIP Address not updated."), *InReceiveIPAddress);
+		return false;
+	}
+} // namespace UE::OSC
+
 UOSCServer::UOSCServer(const FObjectInitializer& ObjectInitializer)
 	: Super(ObjectInitializer)
 {
@@ -92,7 +112,7 @@ void UOSCServer::PostInitProperties()
 	if (DefaultObj != this)
 	{
 		OSCPackets = MakeShared<FPacketQueue>();
-		ServerProxy = MakeUnique<UE::OSC::FServerProxy>(*this);
+		ServerProxy = IServerProxy::Create(this);
 	}
 }
 
@@ -210,7 +230,7 @@ TArray<FOSCAddress> UOSCServer::GetBoundOSCAddressPatterns() const
 
 void UOSCServer::ClearPackets()
 {
-	OSCPackets->Empty();
+	OSCPackets = { };
 }
 
 void UOSCServer::EnqueuePacket(TSharedPtr<UE::OSC::IPacket> InPacket)
@@ -218,24 +238,28 @@ void UOSCServer::EnqueuePacket(TSharedPtr<UE::OSC::IPacket> InPacket)
 	OSCPackets->Enqueue(InPacket);
 }
 
-void UOSCServer::DispatchBundle(const FString& InIPAddress, uint16 InPort, const FOSCBundle& InBundle)
+void UOSCServer::BroadcastBundle(const FOSCBundle& InBundle)
 {
 	using namespace UE::OSC;
 
-	OnOscBundleReceived.Broadcast(InBundle, InIPAddress, InPort);
-	OnOscBundleReceivedNative.Broadcast(InBundle, InIPAddress, InPort);
+	const TSharedRef<IPacket>& Packet = InBundle.GetPacketRef();
+	const FIPv4Endpoint& Endpoint = Packet->GetIPEndpoint();
+	const FString AddrStr = Endpoint.Address.ToString();
+
+	OnOscBundleReceived.Broadcast(InBundle, AddrStr, Endpoint.Port);
+	OnOscBundleReceivedNative.Broadcast(InBundle, AddrStr, Endpoint.Port);
 
-	TSharedPtr<FBundlePacket> BundlePacket = StaticCastSharedRef<FBundlePacket>(InBundle.GetPacketRef());
-	TArray<TSharedRef<IPacket>>& Packets = BundlePacket->GetPackets();
-	for (TSharedRef<IPacket>& Packet : Packets)
+	TSharedRef<FBundlePacket> BundlePacket = StaticCastSharedRef<FBundlePacket>(InBundle.GetPacketRef());
+	const TArray<TSharedRef<IPacket>>& Packets = BundlePacket->GetPackets();
+	for (const TSharedRef<IPacket>& SubPacket : Packets)
 	{
-		if (Packet->IsMessage())
+		if (SubPacket->IsMessage())
 		{
-			DispatchMessage(InIPAddress, InPort, FOSCMessage(Packet));
+			BroadcastMessage(FOSCMessage(SubPacket));
 		}
-		else if (Packet->IsBundle())
+		else if (SubPacket->IsBundle())
 		{
-			DispatchBundle(InIPAddress, InPort, FOSCBundle(Packet));
+			BroadcastBundle(FOSCBundle(SubPacket));
 		}
 		else
 		{
@@ -244,21 +268,27 @@ void UOSCServer::DispatchBundle(const FString& InIPAddress, uint16 InPort, const
 	}
 }
 
-void UOSCServer::DispatchMessage(const FString& InIPAddress, uint16 InPort, const FOSCMessage& InMessage)
+void UOSCServer::BroadcastMessage(const FOSCMessage& InMessage)
 {
-	OnOscMessageReceived.Broadcast(InMessage, InIPAddress, InPort);
-	OnOscMessageReceivedNative.Broadcast(InMessage, InIPAddress, InPort);
+	using namespace UE::OSC;
+
+	const TSharedRef<IPacket>& Packet = InMessage.GetPacketRef();
+	const FIPv4Endpoint& Endpoint = Packet->GetIPEndpoint();
+	const FString AddrStr = Endpoint.Address.ToString();
+
+	OnOscMessageReceived.Broadcast(InMessage, AddrStr, Endpoint.Port);
+	OnOscMessageReceivedNative.Broadcast(InMessage, AddrStr, Endpoint.Port);
 
-	UE_LOG(LogOSC, Verbose, TEXT("Message received from endpoint '%s', OSCAddress of '%s'."), *InIPAddress, *InMessage.GetAddress().GetFullPath());
+	UE_LOG(LogOSC, Verbose, TEXT("Message received from IP endpoint '%s', OSCAddress of '%s'."), *Endpoint.ToString(), *InMessage.GetAddress().GetFullPath());
 
 	for (const TPair<FOSCAddress, FOSCDispatchMessageEvent>& Pair : AddressPatterns)
 	{
 		const FOSCDispatchMessageEvent& DispatchEvent = Pair.Value;
 		if (Pair.Key.Matches(InMessage.GetAddress()))
 		{
-			DispatchEvent.Broadcast(Pair.Key, InMessage, InIPAddress, InPort);
-			UE_LOG(LogOSC, Verbose, TEXT("Message dispatched from endpoint '%s', OSCAddress path of '%s' matched OSCAddress pattern '%s'."),
-				*InIPAddress,
+			DispatchEvent.Broadcast(Pair.Key, InMessage, AddrStr, Endpoint.Port);
+			UE_LOG(LogOSC, Verbose, TEXT("Message dispatched from IP endpoint '%s', OSCAddress path of '%s' matched OSCAddress pattern '%s'."),
+				*Endpoint.ToString(),
 				*InMessage.GetAddress().GetFullPath(),
 				*Pair.Key.GetFullPath());
 		}
@@ -271,21 +301,20 @@ void UOSCServer::PumpPacketQueue()
 
 	check(IsInGameThread());
 
-	TSharedPtr<UE::OSC::IPacket> Packet;
+	TSharedPtr<IPacket> Packet;
 	while (OSCPackets->Dequeue(Packet))
 	{
-		FIPv4Address IPAddr;
-		const FIPv4Endpoint& Endpoint = Packet->GetIPEndpoint();
-		if (ServerProxy->CanProcessPacket(Packet.ToSharedRef()))
+		// Safe to cast to ref as all added items when dispatching to queue should be valid by this point
+		TSharedRef<IPacket> PacketRef = Packet.ToSharedRef();
+		if (ServerProxy->CanProcessPacket(PacketRef))
 		{
-			const FString StringAdr = Endpoint.Address.ToString();
-			if (Packet->IsMessage())
+			if (PacketRef->IsMessage())
 			{
-				DispatchMessage(StringAdr, Endpoint.Port, FOSCMessage(Packet.ToSharedRef()));
+				BroadcastMessage(FOSCMessage(PacketRef));
 			}
-			else if (Packet->IsBundle())
+			else if (PacketRef->IsBundle())
 			{
-				DispatchBundle(StringAdr, Endpoint.Port, FOSCBundle(Packet.ToSharedRef()));
+				BroadcastBundle(FOSCBundle(PacketRef));
 			}
 			else
 			{
diff --git a/Engine/Plugins/Runtime/OSC/Source/OSC/Private/OSCServerProxy.cpp b/Engine/Plugins/Runtime/OSC/Source/OSC/Private/OSCServerProxy.cpp
--- a/Engine/Plugins/Runtime/OSC/Source/OSC/Private/OSCServerProxy.cpp
+++ b/Engine/Plugins/Runtime/OSC/Source/OSC/Private/OSCServerProxy.cpp
@@ -75,6 +75,16 @@ namespace UE::OSC
 		return Endpoint;
 	}
 
+	FString FServerProxy::GetDescription() const
+	{
+		if (Socket)
+		{
+			return Socket->GetDescription();
+		}
+
+		return { };
+	}
+
 	bool FServerProxy::GetMulticastLoopback() const
 	{
 		return bMulticastLoopback;
@@ -130,30 +140,11 @@ namespace UE::OSC
 		}
 	}
 
-	bool FServerProxy::SetAddress(const FString& InReceiveIPAddress, int32 InPort)
-	{
-		if (IsActive())
-		{
-			UE_LOG(LogOSC, Error, TEXT("Cannot set address while OSCServer is active."));
-			return false;
-		}
-
-		FIPv4Address Address;
-		if (!FIPv4Address::Parse(InReceiveIPAddress, Address))
-		{
-			UE_LOG(LogOSC, Error, TEXT("Invalid ReceiveIPAddress '%s'. OSCServer ReceiveIP Address not updated."), *InReceiveIPAddress);
-			return false;
-		}	
-
-		Endpoint = FIPv4Endpoint(Address, InPort);
-		return true;
-	}
-
 	bool FServerProxy::SetIPEndpoint(const FIPv4Endpoint& InEndpoint)
 	{
 		if (IsActive())
 		{
-			UE_LOG(LogOSC, Error, TEXT("Cannot set address while OSCServer is active."));
+			UE_LOG(LogOSC, Error, TEXT("Cannot set '%s' endpoint to '%s' while OSC server is currently active."), *GetDescription(), *InEndpoint.ToString());
 			return false;
 		}
 
@@ -237,14 +228,14 @@ namespace UE::OSC
 
 	void FServerProxy::SetFilterClientsByAllowList(bool bInEnabled)
 	{
-		if (!bInEnabled)
+		if (bInEnabled)
 		{
-			ClientAllowList.Empty(1);
-			ClientAllowList.Add(FIPv4Endpoint::Any);
+			ClientAllowList.Remove(FIPv4Endpoint::Any);
 		}
 		else
 		{
-			ClientAllowList.Remove(FIPv4Endpoint::Any);
+			ClientAllowList.Empty(1);
+			ClientAllowList.Add(FIPv4Endpoint::Any);
 		}
 	}
 } // namespace UE::OSC
diff --git a/Engine/Plugins/Runtime/OSC/Source/OSC/Private/OSCServerProxy.h b/Engine/Plugins/Runtime/OSC/Source/OSC/Private/OSCServerProxy.h
--- a/Engine/Plugins/Runtime/OSC/Source/OSC/Private/OSCServerProxy.h
+++ b/Engine/Plugins/Runtime/OSC/Source/OSC/Private/OSCServerProxy.h
@@ -12,14 +12,22 @@ struct FIPv4Endpoint;
 
 namespace UE::OSC
 {
-	class OSC_API FServerProxy : public IServerProxy
+	class FServerProxy : public IServerProxy , public TSharedFromThis<FServerProxy>
 	{
 	public:
 		FServerProxy(UOSCServer& InServer);
 		virtual ~FServerProxy();
 
 		// Begin IServerProxy interface
+		virtual void AddClientToAllowList(const FString& InIPAddress) override;
+		virtual void AddClientEndpointToAllowList(const FIPv4Endpoint& InIPv4Endpoint) override;
+
 		virtual bool CanProcessPacket(TSharedRef<UE::OSC::IPacket> Packet) const override;
+		virtual void ClearClientEndpointAllowList() override;
+
+		virtual const TSet<FIPv4Endpoint>& GetClientEndpointAllowList() const override;
+		virtual FString GetDescription() const override;
+		virtual const FIPv4Endpoint& GetIPEndpoint() const override;
 		virtual bool GetMulticastLoopback() const override;
 
 		UE_DEPRECATED(5.5, "Use GetIPEndpoint instead")
@@ -28,27 +36,17 @@ namespace UE::OSC
 		UE_DEPRECATED(5.5, "Use GetIPEndpoint instead")
 		virtual int32 GetPort() const override;
 
-		virtual const FIPv4Endpoint& GetIPEndpoint() const override;
-
 		virtual bool IsActive() const override;
 
 		virtual void Listen(const FString& InServerName) override;
 
-		virtual bool SetAddress(const FString& InReceiveIPAddress, int32 InPort) override;
-
 		virtual bool SetIPEndpoint(const FIPv4Endpoint& InEndpoint) override;
 		virtual bool SetMulticastLoopback(bool bInMulticastLoopback) override;
 
 		virtual void Stop() override;
 
-		virtual void AddClientToAllowList(const FString& InIPAddress) override;
 		virtual void RemoveClientFromAllowList(const FString& InIPAddress) override;
-
-
-		virtual void AddClientEndpointToAllowList(const FIPv4Endpoint& InIPv4Endpoint) override;
 		virtual void RemoveClientEndpointFromAllowList(const FIPv4Endpoint& InIPv4Endpoint) override;
-		virtual void ClearClientEndpointAllowList() override;
-		virtual const TSet<FIPv4Endpoint>& GetClientEndpointAllowList() const override;
 
 		virtual void SetFilterClientsByAllowList(bool bInEnabled) override;
 		// End IServerProxy interface
@@ -70,9 +68,9 @@ namespace UE::OSC
 		TSet<FIPv4Endpoint> ClientAllowList;
 
 		/** Endpoint to listen for OSC packets on. If set to 'Any', defaults to LocalHost */
-		FIPv4Endpoint Endpoint;
+		FIPv4Endpoint Endpoint = FIPv4Endpoint::Any;
 
 		/** Whether or not to loopback if address provided is multicast */
 		bool bMulticastLoopback = false;
 	};
-} // namespace UE::OSC	
+} // namespace UE::OSC
diff --git a/Engine/Plugins/Runtime/OSC/Source/OSC/Public/OSCServer.h b/Engine/Plugins/Runtime/OSC/Source/OSC/Public/OSCServer.h
--- a/Engine/Plugins/Runtime/OSC/Source/OSC/Public/OSCServer.h
+++ b/Engine/Plugins/Runtime/OSC/Source/OSC/Public/OSCServer.h
@@ -4,9 +4,12 @@
 
 #include "Async/TaskGraphInterfaces.h"
 #include "Containers/Queue.h"
+#include "Containers/SpscQueue.h"
 #include "Containers/Set.h"
 #include "Containers/Ticker.h"
 #include "Interfaces/IPv4/IPv4Address.h"
+#include "Templates/Function.h"
+#include "Templates/SharedPointer.h"
 #include "Interfaces/IPv4/IPv4Endpoint.h"
 #include "OSCBundle.h"
 #include "OSCMessage.h"
@@ -20,6 +23,7 @@
 
 // Forward Declarations
 class FSocket;
+class UOSCServer;
 
 
 // Delegates
@@ -42,12 +46,17 @@ namespace UE::OSC
 	class OSC_API IServerProxy
 	{
 	public:
+		// Creates a new server proxy that can be used by any system where the provided dispatch callback is called on a worker thread.
+		static TSharedPtr<IServerProxy> Create(UOSCServer* Parent);
 
 		virtual ~IServerProxy() { }
 
 		// Returns whether or not packet can be processed, i.e. is valid and allowlisted.
 		virtual bool CanProcessPacket(TSharedRef<UE::OSC::IPacket> Packet) const = 0;
 
+		// Returns debug description of server proxy
+		virtual FString GetDescription() const = 0;
+
 		UE_DEPRECATED(5.5, "Use GetIPEndpoint instead")
 		virtual FString GetIpAddress() const { return { }; }
 
@@ -66,13 +75,7 @@ namespace UE::OSC
 		virtual void Listen(const FString& ServerName) = 0;
 
 		UE_DEPRECATED(5.5, "Use SetIPEndpoint instead")
-		virtual bool SetAddress(const FString& InReceiveIPAddress, int32 InPort)
-		{
-			FIPv4Endpoint Endpoint;
-			FIPv4Address::Parse(InReceiveIPAddress, Endpoint.Address);
-			Endpoint.Port = InPort;
-			return SetIPEndpoint(Endpoint);
-		}
+		virtual bool SetAddress(const FString& InReceiveIPAddress, int32 InPort);
 
 		// Sets the current server's endpoint.  Ignores request and returns false if server
 		// is currently active.
@@ -234,7 +237,7 @@ public:
 	void SetTickInEditor(bool bInTickInEditor);
 #endif // WITH_EDITOR
 
-	UE_DEPRECATED(5.5, "Clearing packets directly is not thread safe and no longer supported.")
+	UE_DEPRECATED(5.5, "Clearing packets directly is not thread-safe and no longer supported.")
 	void ClearPackets();
 
 	void EnqueuePacket(TSharedPtr<UE::OSC::IPacket> InPacket);
@@ -247,17 +250,17 @@ protected:
 	virtual void PostInitProperties() override;
 
 private:
-	using FPacketQueue = TQueue<TSharedPtr<UE::OSC::IPacket>>;
+	using FPacketQueue = TSpscQueue<TSharedPtr<UE::OSC::IPacket>>;
 	void PumpPacketQueue();
 
-	/** Dispatches provided bundle received */
-	void DispatchBundle(const FString& InIPAddress, uint16 InPort, const FOSCBundle& InBundle);
+	/** Broadcasts provided bundle received to be dispatched on the GameThread */
+	void BroadcastBundle(const FOSCBundle& InBundle);
 
-	/** Dispatches provided message received */
-	void DispatchMessage(const FString& InIPAddress, uint16 InPort, const FOSCMessage& InMessage);
+	/** Broadcasts provided message received to be dispatched on the GameThread */
+	void BroadcastMessage(const FOSCMessage& InMessage);
 
 	/** Pointer to internal implementation of server proxy */
-	TUniquePtr<UE::OSC::IServerProxy> ServerProxy;
+	TSharedPtr<UE::OSC::IServerProxy> ServerProxy;
 
 	/** Queue stores incoming OSC packet requests to process on the game thread. */
 	TSharedPtr<FPacketQueue> OSCPackets;
