- Migrate OSC Server away from generic UDP receiver to more modernized OSCServerReceiver that is RAII and utilizes unique socket
- Provide ability for client system to provide dispatch function, allowing for OSC Server Proxies to operate independent of UObject & respond to arbitrary processing thread
- Fix OSC Server loopback address not being observed on socket/socket receiver generation
- Re-implement ClearPackets in a thread-safe way as internal implementation to OSCServer and call when stopping to avoid passing along packets that would then be invalid to a new server proxy implementation

#rnx

diff --git a/Engine/Plugins/Runtime/OSC/Source/OSC/Private/OSCServer.cpp b/Engine/Plugins/Runtime/OSC/Source/OSC/Private/OSCServer.cpp
--- a/Engine/Plugins/Runtime/OSC/Source/OSC/Private/OSCServer.cpp
+++ b/Engine/Plugins/Runtime/OSC/Source/OSC/Private/OSCServer.cpp
@@ -1,31 +1,30 @@
 // Copyright Epic Games, Inc. All Rights Reserved.
 #include "OSCServer.h"
 
-
-#include "OSCMessage.h"
+#include "Containers/Ticker.h"
+#include "OSCAddress.h"
 #include "OSCBundle.h"
 #include "OSCBundlePacket.h"
 #include "OSCLog.h"
+#include "OSCMessage.h"
+#include "OSCMessagePacket.h"
+#include "OSCPacket.h"
 #include "OSCServerProxy.h"
 
 
 namespace UE::OSC
 {
-	TSharedPtr<IServerProxy> IServerProxy::Create(UOSCServer* Parent)
+	TSharedPtr<IServerProxy> IServerProxy::Create()
 	{
-		return MakeShared<OSC::FServerProxy>(*Parent);
+		return MakeShared<OSC::FServerProxy>();
 	}
 
 	bool IServerProxy::SetAddress(const FString& InReceiveIPAddress, int32 InPort)
 	{
-		FIPv4Address Address;
-		if (FIPv4Address::Parse(InReceiveIPAddress, Address))
-		{
-			return SetIPEndpoint(FIPv4Endpoint(Address, InPort));
-		}
-
-		UE_LOG(LogOSC, Error, TEXT("Invalid ReceiveIPAddress '%s'. OSCServer ReceiveIP Address not updated."), *InReceiveIPAddress);
-		return false;
+		FIPv4Endpoint Endpoint;
+		FIPv4Address::Parse(InReceiveIPAddress, Endpoint.Address);
+		Endpoint.Port = InPort;
+		return SetIPEndpoint(Endpoint);
 	}
 } // namespace UE::OSC
 
@@ -49,6 +48,8 @@ bool UOSCServer::IsActive() const
 void UOSCServer::Listen()
 {
 	check(ServerProxy.IsValid());
+
+	ClearPacketsInternal();
 	ServerProxy->Listen(GetName());
 
 	TickHandle = FTSTicker::GetCoreTicker().AddTicker(FTickerDelegate::CreateLambda([this](float /*Time*/)
@@ -87,8 +88,11 @@ void UOSCServer::Stop()
 	// CDO May have this not set by initializer ctor, so have to check if valid
 	if (ServerProxy.IsValid())
 	{
+		ServerProxy->SetOnDispatchPacket({ });
 		ServerProxy->Stop();
 	}
+
+	ClearPacketsInternal();
 }
 
 void UOSCServer::BeginDestroy()
@@ -112,7 +116,7 @@ void UOSCServer::PostInitProperties()
 	if (DefaultObj != this)
 	{
 		OSCPackets = MakeShared<FPacketQueue>();
-		ServerProxy = IServerProxy::Create(this);
+		ServerProxy = IServerProxy::Create();
 	}
 }
 
@@ -228,14 +232,19 @@ TArray<FOSCAddress> UOSCServer::GetBoundOSCAddressPatterns() const
 	return OutAddressPatterns;
 }
 
-void UOSCServer::ClearPackets()
+void UOSCServer::ClearPacketsInternal()
 {
-	OSCPackets = { };
-}
+	using namespace UE::OSC;
 
-void UOSCServer::EnqueuePacket(TSharedPtr<UE::OSC::IPacket> InPacket)
-{
-	OSCPackets->Enqueue(InPacket);
+	OSCPackets = MakeShared<FPacketQueue>();
+
+	if (ServerProxy.IsValid())
+	{
+		ServerProxy->SetOnDispatchPacket(MakeShared<IServerProxy::FOnDispatchPacket>([Queue = OSCPackets](TSharedRef<UE::OSC::IPacket> Packet)
+		{
+			Queue->Enqueue(Packet);
+		}));
+	}
 }
 
 void UOSCServer::BroadcastBundle(const FOSCBundle& InBundle)
diff --git a/Engine/Plugins/Runtime/OSC/Source/OSC/Private/OSCServerProxy.cpp b/Engine/Plugins/Runtime/OSC/Source/OSC/Private/OSCServerProxy.cpp
--- a/Engine/Plugins/Runtime/OSC/Source/OSC/Private/OSCServerProxy.cpp
+++ b/Engine/Plugins/Runtime/OSC/Source/OSC/Private/OSCServerProxy.cpp
@@ -10,11 +10,7 @@
 
 namespace UE::OSC
 {
-	FServerProxy::FServerProxy(UOSCServer& InServer)
-		: Server(&InServer)
-		, Socket(nullptr)
-		, SocketReceiver(nullptr)
-		, bMulticastLoopback(false)
+	FServerProxy::FServerProxy()
 	{
 		ClientAllowList.Add(FIPv4Endpoint::Any);
 	}
@@ -44,11 +40,9 @@ namespace UE::OSC
 		return ClientAllowList.Contains(EndpointToTest);
 	}
 
-	void FServerProxy::OnPacketReceived(const FArrayReaderPtr& InData, const FIPv4Endpoint& InEndpoint)
+	void FServerProxy::OnPacketReceived(FConstPacketDataRef InData, const FIPv4Endpoint& InEndpoint)
 	{
-		using namespace UE::OSC;
-
-		TSharedPtr<IPacket> Packet = IPacket::CreatePacket(InData->GetData(), InEndpoint);
+		TSharedPtr<UE::OSC::IPacket> Packet = UE::OSC::IPacket::CreatePacket(InData->GetData(), InEndpoint);
 		if (!Packet.IsValid())
 		{
 			UE_LOG(LogOSC, Verbose, TEXT("Message received from endpoint '%s' invalid OSC packet."), *InEndpoint.ToString());
@@ -57,7 +51,12 @@ namespace UE::OSC
 
 		FStream Stream = FStream(InData->GetData(), InData->Num());
 		Packet->ReadData(Stream);
-		Server->EnqueuePacket(Packet);
+
+		FScopeLock Lock(&MutateDispatchFunctionCritSec);
+		if (OnDispatchFunction.IsValid())
+		{
+			(*OnDispatchFunction)(Packet.ToSharedRef());
+		}
 	}
 
 	FString FServerProxy::GetIpAddress() const
@@ -77,9 +76,9 @@ namespace UE::OSC
 
 	FString FServerProxy::GetDescription() const
 	{
-		if (Socket)
+		if (ServerReceiver.IsValid())
 		{
-			return Socket->GetDescription();
+			return ServerReceiver->GetDescription();
 		}
 
 		return { };
@@ -92,52 +91,23 @@ namespace UE::OSC
 
 	bool FServerProxy::IsActive() const
 	{
-		return SocketReceiver != nullptr;
+		return ServerReceiver.IsValid();
 	}
 
 	void FServerProxy::Listen(const FString& InServerName)
 	{
 		if (IsActive())
 		{
-			UE_LOG(LogOSC, Error, TEXT("OSCServer currently listening: %s:%d. Failed to start new service prior to calling stop."),
-				*InServerName, *Endpoint.Address.ToString(), Endpoint.Port);
+			UE_LOG(LogOSC, Error, TEXT("OSCServer currently '%s' currently listening to endpoint '%s'. Failed to start new service prior to calling stop."),
+				*InServerName, *Endpoint.Address.ToString());
 			return;
 		}
 
-		FUdpSocketBuilder Builder(*InServerName);
-		Builder.BoundToPort(Endpoint.Port);
-		if (Endpoint.Address.IsMulticastAddress())
-		{
-			Builder.JoinedToGroup(Endpoint.Address);
-			if (bMulticastLoopback)
-			{
-				Builder.WithMulticastLoopback();
-			}
-		}
-		else
-		{
-			if (bMulticastLoopback)
-			{
-				UE_LOG(LogOSC, Warning, TEXT("OSCServer '%s' ReceiveIPAddress provided is not a multicast address.  Not respecting MulticastLoopback boolean."),
-					*InServerName);
-			}
-			Builder.BoundToAddress(Endpoint.Address);
-		}
-
-		Socket = Builder.Build();
-		if (Socket)
-		{
-			SocketReceiver = new FUdpSocketReceiver(Socket, FTimespan::FromMilliseconds(100), *(InServerName + TEXT("_ListenerThread")));
-			SocketReceiver->OnDataReceived().BindRaw(this, &FServerProxy::OnPacketReceived);
-			SocketReceiver->Start();
+		FServerReceiver::FOptions Options;
+		Options.bMulticastLoopback = bMulticastLoopback;
+		Options.ReceivedDataDelegate.BindSP(this, &FServerProxy::OnPacketReceived);
 
-			UE_LOG(LogOSC, Display, TEXT("OSCServer '%s' Listening: %s:%d."), *InServerName, *Endpoint.Address.ToString(), Endpoint.Port);
-		}
-		else
-		{
-			// This is expected when the server isn't available, so it's not a Warning
-			UE_LOG(LogOSC, Display, TEXT("OSCServer '%s' failed to bind to socket on %s:%d. Check that the server is available on the specified address."), *InServerName, *Endpoint.Address.ToString(), Endpoint.Port);
-		}
+		ServerReceiver = FServerReceiver::Launch(InServerName, Endpoint, MoveTemp(Options));
 	}
 
 	bool FServerProxy::SetIPEndpoint(const FIPv4Endpoint& InEndpoint)
@@ -164,22 +134,15 @@ namespace UE::OSC
 		return true;
 	}
 
-	void FServerProxy::Stop()
+	void FServerProxy::SetOnDispatchPacket(TSharedPtr<FOnDispatchPacket> OnDispatch)
 	{
-		if (SocketReceiver)
-		{
-			delete SocketReceiver;
-			SocketReceiver = nullptr;
-		}
-
-		if (Socket)
-		{
-			Socket->Close();
-			ISocketSubsystem::Get(PLATFORM_SOCKETSUBSYSTEM)->DestroySocket(Socket);
-			Socket = nullptr;
-		}
-
+		FScopeLock Lock(&MutateDispatchFunctionCritSec);
+		OnDispatchFunction = OnDispatch;
+	}
 
+	void FServerProxy::Stop()
+	{
+		ServerReceiver.Reset();
 	}
 
 	void FServerProxy::AddClientToAllowList(const FString& InIPAddress)
diff --git a/Engine/Plugins/Runtime/OSC/Source/OSC/Private/OSCServerProxy.h b/Engine/Plugins/Runtime/OSC/Source/OSC/Private/OSCServerProxy.h
--- a/Engine/Plugins/Runtime/OSC/Source/OSC/Private/OSCServerProxy.h
+++ b/Engine/Plugins/Runtime/OSC/Source/OSC/Private/OSCServerProxy.h
@@ -2,10 +2,10 @@
 #pragma once
 
 
-#include "Common/UdpSocketReceiver.h"
-#include "Tickable.h"
+#include "HAL/CriticalSection.h"
 
 #include "OSCServer.h"
+#include "OSCServerReceiver.h"
 
 struct FIPv4Endpoint;
 
@@ -15,7 +15,7 @@ namespace UE::OSC
 	class FServerProxy : public IServerProxy , public TSharedFromThis<FServerProxy>
 	{
 	public:
-		FServerProxy(UOSCServer& InServer);
+		FServerProxy();
 		virtual ~FServerProxy();
 
 		// Begin IServerProxy interface
@@ -42,7 +42,7 @@ namespace UE::OSC
 
 		virtual bool SetIPEndpoint(const FIPv4Endpoint& InEndpoint) override;
 		virtual bool SetMulticastLoopback(bool bInMulticastLoopback) override;
-
+		virtual void SetOnDispatchPacket(TSharedPtr<FOnDispatchPacket> OnDispatch) override;
 		virtual void Stop() override;
 
 		virtual void RemoveClientFromAllowList(const FString& InIPAddress) override;
@@ -53,16 +53,11 @@ namespace UE::OSC
 
 	private:
 		/** Callback that receives data from a socket. */
-		void OnPacketReceived(const FArrayReaderPtr& InData, const FIPv4Endpoint& InEndpoint);
-
-		/** Parent server UObject */
-		UOSCServer* Server;
+		void OnPacketReceived(FConstPacketDataRef InData, const FIPv4Endpoint& InEndpoint);
 
-		/** Socket used to listen for OSC packets. */
-		FSocket* Socket;
+		mutable FCriticalSection MutateDispatchFunctionCritSec;
 
-		/** UDP receiver. */
-		FUdpSocketReceiver* SocketReceiver;
+		TSharedPtr<FServerReceiver> ServerReceiver;
 
 		/** Only packets from this list of client addresses will be processed if bFilterClientsByAllowList is true. */
 		TSet<FIPv4Endpoint> ClientAllowList;
@@ -72,5 +67,7 @@ namespace UE::OSC
 
 		/** Whether or not to loopback if address provided is multicast */
 		bool bMulticastLoopback = false;
+
+		TSharedPtr<FOnDispatchPacket> OnDispatchFunction;
 	};
 } // namespace UE::OSC
diff --git a/Engine/Plugins/Runtime/OSC/Source/OSC/Private/OSCServerReceiver.cpp b/Engine/Plugins/Runtime/OSC/Source/OSC/Private/OSCServerReceiver.cpp
new file mode 100644
--- /dev/null
+++ b/Engine/Plugins/Runtime/OSC/Source/OSC/Private/OSCServerReceiver.cpp
@@ -0,0 +1,242 @@
+// Copyright Epic Games, Inc. All Rights Reserved.
+
+#include "OSCServerReceiver.h"
+
+
+namespace UE::OSC
+{
+	namespace ServerReceieverPrivate
+	{
+		ISocketSubsystem& GetSocketSubsystemChecked()
+		{
+			ISocketSubsystem* Subsystem = ISocketSubsystem::Get(PLATFORM_SOCKETSUBSYSTEM);
+			check(Subsystem);
+			return *Subsystem;
+		}
+
+		void SetAsLocalHostIfLoopback(ISocketSubsystem& InSocketSubsystem, FIPv4Endpoint& OutEndpoint)
+		{
+			if (OutEndpoint.Address.IsLoopbackAddress())
+			{
+				check(GLog);
+
+				bool bCanBind = false;
+				TSharedRef<FInternetAddr> LocalIP = InSocketSubsystem.GetLocalHostAddr(*GLog, bCanBind);
+				uint32 LocalIPAddr = 0;
+				LocalIP->GetIp(LocalIPAddr);
+				OutEndpoint.Address = FIPv4Address(LocalIPAddr);
+			}
+		}
+	} // namespace ServerReceieverPrivate
+
+	FServerReceiver::FServerReceiver(FPrivateToken, FOptions&& InOptions)
+		: Options(MoveTemp(InOptions))
+	{
+	}
+
+	FServerReceiver::~FServerReceiver()
+	{
+		bStopping = true;
+
+		UE_LOG(LogOSC, Display, TEXT("Destroying OSC Socket Receiver '%s'"), *GetDescription());
+		if (Thread)
+		{
+			const FString& ThreadName = Thread->GetThreadName();
+			UE_LOG(LogOSC, Display, TEXT("Killing process thread '%s'..."), *ThreadName);
+			constexpr bool bWait = true;
+			Thread->Kill(bWait);
+			UE_LOG(LogOSC, Display, TEXT("Process thread '%s' killed successfully."), *ThreadName);
+			Thread.Reset();
+		}
+
+		if (Socket)
+		{
+			const FString SocketDesc = Socket->GetDescription();
+			const bool bSocketClosed = Socket->Close();
+			UE_LOG(LogOSC, Display, TEXT("Socket '%s' '%s' closed."), *SocketDesc, bSocketClosed ? TEXT("successfully") : TEXT("failed to"));
+			Socket.Reset();
+		}
+	}
+
+	bool FServerReceiver::BindSocket(ISocketSubsystem& SocketSubsystem, const FIPv4Endpoint& InEndpoint, TSharedRef<FInternetAddr> RemoteAddr, const FString& InName)
+	{
+		check(Socket);
+
+		if (!Socket->SetNonBlocking(true) ||
+			!Socket->SetReuseAddr(false) ||
+			!Socket->SetBroadcast(false) ||
+			!Socket->SetRecvErr())
+		{
+			const ESocketErrors SocketError = SocketSubsystem.GetLastErrorCode();
+
+			UE_LOG(LogOSC, Warning, TEXT("OSC Socket Receiver: Failed to configure socket %s: Error code %d."),
+				*InName, int32(SocketError));
+
+			return false;
+		}
+
+		if (!Socket->Bind(*RemoteAddr))
+		{
+			const ESocketErrors SocketError = SocketSubsystem.GetLastErrorCode();
+
+			UE_LOG(LogOSC, Warning, TEXT("OSC Socket Receiver: Failed to bind %s to %s. Error code %d."),
+				*InName, *InEndpoint.ToString(), int32(SocketError));
+
+			return false;
+		}
+
+		return true;
+	}
+
+	bool FServerReceiver::InitMulticast(ISocketSubsystem& InSocketSubsystem, const FIPv4Endpoint& InEndpoint, TSharedRef<FInternetAddr> RemoteAddr, const FString& InName)
+	{
+		if (!Socket->SetMulticastLoopback(Options.bMulticastLoopback) || !Socket->SetMulticastTtl(1))
+		{
+			const ESocketErrors SocketError = InSocketSubsystem.GetLastErrorCode();
+			UE_LOG(LogOSC, Warning, TEXT("OSC Socket Receiver: Failed to configure multicast for %s (loopback: %i). Error code %d."),
+				*InName, Options.bMulticastLoopback, int32(SocketError));
+			return false;
+		}
+
+		TSharedRef<FInternetAddr> MulticastAddress = InSocketSubsystem.CreateInternetAddr(RemoteAddr->GetProtocolType());
+		MulticastAddress->SetBroadcastAddress();
+
+		TSharedPtr<FInternetAddr> AddressToUse;
+		if (InEndpoint.Address.IsSessionFrontendMulticast() && MulticastAddress->GetProtocolType() != FNetworkProtocolTypes::IPv4)
+		{
+			AddressToUse = MulticastAddress;
+		}
+		else
+		{
+			AddressToUse = FIPv4Endpoint(InEndpoint.Address, 0).ToInternetAddr();
+		}
+
+		if (!Socket->JoinMulticastGroup(*AddressToUse, *FIPv4Endpoint(FIPv4Address::Any, 0).ToInternetAddr()))
+		{
+			const ESocketErrors SocketError = InSocketSubsystem.GetLastErrorCode();
+			UE_LOG(LogOSC, Warning, TEXT("OSC Server Receiver: Failed to subscribe %s to multicast group %s. Error code %d."),
+				*InName, *InEndpoint.Address.ToString(), int32(SocketError));
+			return false;
+		}
+
+		return true;
+	}
+
+	bool FServerReceiver::InitSocket(const FString& InName, FIPv4Endpoint Endpoint)
+	{
+		using namespace ServerReceieverPrivate;
+
+		ISocketSubsystem& SocketSubsystem = GetSocketSubsystemChecked();
+		SetAsLocalHostIfLoopback(SocketSubsystem, Endpoint);
+
+		TSharedRef<FInternetAddr> RemoteAddr = Endpoint.ToInternetAddr();
+		Socket = SocketSubsystem.CreateUniqueSocket(NAME_DGram, *InName, RemoteAddr->GetProtocolType());
+		if (!Socket.IsValid())
+		{
+			UE_LOG(LogOSC, Warning, TEXT("OSCServer '%s' failed to create socket"), *InName);
+			return false;
+		}
+
+		const bool bSocketBound = BindSocket(SocketSubsystem, Endpoint, RemoteAddr, InName);
+		if (!bSocketBound)
+		{
+			return false;
+		}
+
+		if (Endpoint.Address.IsMulticastAddress())
+		{
+			const bool bMulticastInit = InitMulticast(SocketSubsystem, Endpoint, RemoteAddr, InName);
+			if (!bMulticastInit)
+			{
+				return false;
+			}
+		}
+
+		return true;
+	}
+
+	FString FServerReceiver::GetDescription() const
+	{
+		if (Socket)
+		{
+			return Socket->GetDescription();
+		}
+
+		return InvalidReceiverName;
+	}
+
+	TSharedRef<FServerReceiver> FServerReceiver::Launch(const FString& InName, const FIPv4Endpoint& InEndpoint, FOptions&& InOptions)
+	{
+		TSharedRef<FServerReceiver> NewReceiver = MakeShared<FServerReceiver>(FPrivateToken { }, MoveTemp(InOptions));
+
+		const bool bSocketInit = NewReceiver->InitSocket(InName, InEndpoint);
+		if (bSocketInit)
+		{
+			if (NewReceiver->Socket)
+			{
+				NewReceiver->StartThread(InName, InEndpoint);
+			}
+		}
+		else
+		{
+			NewReceiver->Socket.Reset();
+		}
+
+		return NewReceiver;
+	}
+
+	uint32 FServerReceiver::Run()
+	{
+		check(Socket);
+		while (!bStopping)
+		{
+			if (Socket->Wait(ESocketWaitConditions::WaitForRead, Options.WaitTime))
+			{
+				Tick();
+			}
+		}
+
+		return 0;
+	}
+
+	void FServerReceiver::StartThread(const FString& InName, const FIPv4Endpoint& InEndpoint)
+	{
+		check(Socket);
+		check(Socket->GetSocketType() == SOCKTYPE_Datagram);
+
+		const FString ThreadName = FString::Printf(TEXT("OSCReceiver_%s_%s"), *InName, *InEndpoint.ToString());
+		Thread = TSharedPtr<FRunnableThread>(FRunnableThread::Create(
+			this,
+			*ThreadName,
+			Options.StackSize,
+			Options.Priority,
+			FPlatformAffinity::GetPoolThreadMask()
+		));
+	}
+
+	void FServerReceiver::Tick()
+	{
+		using namespace ServerReceieverPrivate;
+
+		TRACE_CPUPROFILER_EVENT_SCOPE(FServerReceiver_Tick);
+
+		checkf(Socket, TEXT("ServerReceiver thread should never execute when socket is invalid"));
+
+		uint32 Size;
+		while (Socket->HasPendingData(Size))
+		{
+			FPacketDataRef PacketData = MakeShared<TArray<uint8>>();
+			PacketData->AddUninitialized(FMath::Min(Size, Options.MaxReadBufferSize));
+
+			TSharedRef<FInternetAddr> Sender = GetSocketSubsystemChecked().CreateInternetAddr();
+
+			int32 NumRead = 0;
+			if (Socket->RecvFrom(PacketData->GetData(), PacketData->Num(), NumRead, *Sender))
+			{
+				ensureMsgf((uint32)NumRead <= Options.MaxReadBufferSize, TEXT("OSC Server Socket '%s' overflow"), *Socket->GetDescription());
+				PacketData->RemoveAt(NumRead, PacketData->Num() - NumRead, EAllowShrinking::No);
+				Options.ReceivedDataDelegate.ExecuteIfBound(StaticCastSharedRef<const TArray<uint8>>(PacketData), FIPv4Endpoint(Sender));
+			}
+		}
+	}
+} // namespace UE::OSC
diff --git a/Engine/Plugins/Runtime/OSC/Source/OSC/Private/OSCServerReceiver.h b/Engine/Plugins/Runtime/OSC/Source/OSC/Private/OSCServerReceiver.h
new file mode 100644
--- /dev/null
+++ b/Engine/Plugins/Runtime/OSC/Source/OSC/Private/OSCServerReceiver.h
@@ -0,0 +1,74 @@
+// Copyright Epic Games, Inc. All Rights Reserved.
+#pragma once
+
+#include "GenericPlatform/GenericPlatformAffinity.h"
+#include "HAL/Runnable.h"
+#include "HAL/RunnableThread.h"
+#include "Interfaces/IPv4/IPv4Endpoint.h"
+#include "IPAddress.h"
+#include "Misc/ConfigCacheIni.h"
+#include "Misc/SingleThreadRunnable.h"
+#include "OSCLog.h"
+#include "OSCStream.h"
+#include "OSCTypes.h"
+#include "Sockets.h"
+#include "UObject/NoExportTypes.h"
+
+#include <atomic>
+
+
+namespace UE::OSC
+{
+	using FPacketDataRef = TSharedRef<TArray<uint8>>;
+	using FConstPacketDataRef = TSharedRef<const TArray<uint8>>;
+	DECLARE_DELEGATE_TwoParams(FOnServerReceivedData, FConstPacketDataRef, const FIPv4Endpoint&);
+
+	class FServerReceiver
+		: public TSharedFromThis<FServerReceiver>
+		, public FRunnable
+		, private FSingleThreadRunnable
+	{
+		struct FPrivateToken { explicit FPrivateToken() = default; };
+
+	public:
+		struct FOptions
+		{
+			FOnServerReceivedData ReceivedDataDelegate;
+			bool bMulticastLoopback = false;
+
+			EThreadPriority Priority = EThreadPriority::TPri_AboveNormal;
+			int32 StackSize = 128 * 1024;
+			FTimespan WaitTime = FTimespan::FromMilliseconds(100);
+			uint32 MaxReadBufferSize = 65507u;
+		};
+
+		explicit FServerReceiver(FPrivateToken, FOptions&& InOptions);
+		virtual ~FServerReceiver();
+
+		static TSharedRef<FServerReceiver> Launch(const FString& InName, const FIPv4Endpoint& InEndpoint, FOptions&& InOptions);
+
+		FString GetDescription() const;
+
+	protected:
+		virtual uint32 Run() override;
+		virtual void Tick() override;
+
+	private:
+		bool BindSocket(ISocketSubsystem& SocketSubsystem, const FIPv4Endpoint& InEndpoint, TSharedRef<FInternetAddr> RemoteAddr, const FString& InName);
+		bool InitMulticast(ISocketSubsystem& SocketSubsystem, const FIPv4Endpoint& InEndpoint, TSharedRef<FInternetAddr> RemoteAddr, const FString& InName);
+		bool InitSocket(const FString& InName, FIPv4Endpoint InEndpoint);
+		void StartThread(const FString& InName, const FIPv4Endpoint& InEndpoint);
+
+		std::atomic<bool> bStopping { false };
+
+		/** The network socket. */
+		FUniqueSocket Socket;
+
+		/** Dedicated thread for server to run on. */
+		TSharedPtr<FRunnableThread> Thread;
+
+		FString InvalidReceiverName;
+
+		const FOptions Options;
+	};
+} // namespace UE::OSC
diff --git a/Engine/Plugins/Runtime/OSC/Source/OSC/Public/OSCServer.h b/Engine/Plugins/Runtime/OSC/Source/OSC/Public/OSCServer.h
--- a/Engine/Plugins/Runtime/OSC/Source/OSC/Public/OSCServer.h
+++ b/Engine/Plugins/Runtime/OSC/Source/OSC/Public/OSCServer.h
@@ -46,8 +46,10 @@ namespace UE::OSC
 	class OSC_API IServerProxy
 	{
 	public:
+		using FOnDispatchPacket = TUniqueFunction<void(TSharedRef<UE::OSC::IPacket>)>;
+
 		// Creates a new server proxy that can be used by any system where the provided dispatch callback is called on a worker thread.
-		static TSharedPtr<IServerProxy> Create(UOSCServer* Parent);
+		static TSharedPtr<IServerProxy> Create();
 
 		virtual ~IServerProxy() { }
 
@@ -84,8 +86,13 @@ namespace UE::OSC
 		// Sets whether or not loopback is enabled.  Returns false and request is ignored
 		// if server is currently active.
 		virtual bool SetMulticastLoopback(bool bInMulticastLoopback) = 0;
+
+		// Sets dispatch function to be called when OSC packet is received (thread safe
+		// and can be mutated while server is running)
+		virtual void SetOnDispatchPacket(TSharedPtr<FOnDispatchPacket> OnDispatch) { };
+
 #if WITH_EDITOR
-		UE_DEPRECATED(5.5, "All server proxies are now ticked internally by handle in all editor contexts.")
+		UE_DEPRECATED(5.5, "ServerProxies are no longer independently ticked objects only and can now be accessed by threads other than the GameThread via the 'SetOnDispatchPacket' setter.")
 		virtual void SetTickableInEditor(bool bInTickInEditor) { };
 #endif // WITH_EDITOR
 
@@ -240,7 +247,8 @@ public:
 	UE_DEPRECATED(5.5, "Clearing packets directly is not thread-safe and no longer supported.")
 	void ClearPackets();
 
-	void EnqueuePacket(TSharedPtr<UE::OSC::IPacket> InPacket);
+	UE_DEPRECATED(5.5, "Enqueuing packets is now handled privately")
+	void EnqueuePacket(TSharedPtr<UE::OSC::IPacket> InPacket) { }
 
 	UE_DEPRECATED(5.5, "Pumping packets is now handled privately")
 	void PumpPacketQueue(const TSet<uint32>* InAllowlistedClients) { }
@@ -250,6 +258,8 @@ protected:
 	virtual void PostInitProperties() override;
 
 private:
+	void ClearPacketsInternal();
+
 	using FPacketQueue = TSpscQueue<TSharedPtr<UE::OSC::IPacket>>;
 	void PumpPacketQueue();
 
