More OSC clean-up/fix-ups from earlier review feedback
#rnx

diff --git a/Engine/Plugins/Runtime/OSC/Source/OSC/Private/OSCClientProxy.cpp b/Engine/Plugins/Runtime/OSC/Source/OSC/Private/OSCClientProxy.cpp
--- a/Engine/Plugins/Runtime/OSC/Source/OSC/Private/OSCClientProxy.cpp
+++ b/Engine/Plugins/Runtime/OSC/Source/OSC/Private/OSCClientProxy.cpp
@@ -95,13 +95,13 @@ namespace UE::OSC
 		}
 	}
 
-	void FClientProxy::SendMessage(FOSCMessage& Message)
+	void FClientProxy::SendMessage(const FOSCMessage& Message)
 	{
 		const TSharedRef<UE::OSC::IPacket>& Packet = Message.GetPacketRef();
 		SendPacket(Packet.Get());
 	}
 
-	void FClientProxy::SendBundle(FOSCBundle& Bundle)
+	void FClientProxy::SendBundle(const FOSCBundle& Bundle)
 	{
 		const TSharedRef<UE::OSC::IPacket>& Packet = Bundle.GetPacketRef();
 		SendPacket(Packet.Get());
diff --git a/Engine/Plugins/Runtime/OSC/Source/OSC/Private/OSCClientProxy.h b/Engine/Plugins/Runtime/OSC/Source/OSC/Private/OSCClientProxy.h
--- a/Engine/Plugins/Runtime/OSC/Source/OSC/Private/OSCClientProxy.h
+++ b/Engine/Plugins/Runtime/OSC/Source/OSC/Private/OSCClientProxy.h
@@ -31,8 +31,8 @@ namespace UE::OSC
 
 		bool IsActive() const override;
 
-		void SendMessage(FOSCMessage& Message) override;
-		void SendBundle(FOSCBundle& Bundle) override;
+		void SendMessage(const FOSCMessage& Message) override;
+		void SendBundle(const FOSCBundle& Bundle) override;
 
 		void Stop() override;
 
diff --git a/Engine/Plugins/Runtime/OSC/Source/OSC/Private/OSCMessage.cpp b/Engine/Plugins/Runtime/OSC/Source/OSC/Private/OSCMessage.cpp
--- a/Engine/Plugins/Runtime/OSC/Source/OSC/Private/OSCMessage.cpp
+++ b/Engine/Plugins/Runtime/OSC/Source/OSC/Private/OSCMessage.cpp
@@ -10,6 +10,15 @@ FOSCMessage::FOSCMessage()
 {
 }
 
+FOSCMessage::FOSCMessage(FOSCAddress Address, TArray<UE::OSC::FOSCData> Args)
+	: Packet(MakeShared<UE::OSC::FMessagePacket>())
+{
+	using namespace UE::OSC;
+	TSharedRef<FMessagePacket> MsgPacket = StaticCastSharedRef<FMessagePacket>(Packet);
+	MsgPacket->SetAddress(MoveTemp(Address));
+	MsgPacket->SetArguments(MoveTemp(Args));
+}
+
 FOSCMessage::FOSCMessage(const TSharedRef<UE::OSC::IPacket>& InPacket)
 	: Packet(InPacket)
 {
@@ -18,13 +27,12 @@ FOSCMessage::FOSCMessage(const TSharedRef<UE::OSC::IPacket>& InPacket)
 FOSCMessage::FOSCMessage(const TSharedPtr<UE::OSC::IPacket>& InPacket)
 	: FOSCMessage::FOSCMessage()
 {
-	Packet = TSharedRef<UE::OSC::IPacket>(InPacket.Get());
+	Packet = InPacket.ToSharedRef();
 }
 
 void FOSCMessage::SetPacket(TSharedPtr<UE::OSC::IPacket>& InPacket)
 {
-	using namespace UE::OSC;
-	Packet = TSharedRef<IPacket>(InPacket.Get());
+	Packet = InPacket.ToSharedRef();
 }
 
 void FOSCMessage::SetPacket(TSharedRef<UE::OSC::IPacket>& InPacket)
diff --git a/Engine/Plugins/Runtime/OSC/Source/OSC/Private/OSCMessagePacket.cpp b/Engine/Plugins/Runtime/OSC/Source/OSC/Private/OSCMessagePacket.cpp
--- a/Engine/Plugins/Runtime/OSC/Source/OSC/Private/OSCMessagePacket.cpp
+++ b/Engine/Plugins/Runtime/OSC/Source/OSC/Private/OSCMessagePacket.cpp
@@ -23,9 +23,9 @@ namespace UE::OSC
 	{
 	}
 
-	void FMessagePacket::AddArgument(FOSCData&& InOSCData)
+	void FMessagePacket::AddArgument(FOSCData OSCData)
 	{
-		Arguments.Add(MoveTemp(InOSCData));
+		Arguments.Add(MoveTemp(OSCData));
 	}
 
 	void FMessagePacket::EmptyArguments()
@@ -43,14 +43,14 @@ namespace UE::OSC
 		return Arguments;
 	}
 
-	void FMessagePacket::SetAddress(const FOSCAddress& InAddress)
+	void FMessagePacket::SetAddress(FOSCAddress InAddress)
 	{
-		Address = InAddress;
+		Address = MoveTemp(InAddress);
 	}
 
-	void FMessagePacket::SetAddress(FOSCAddress&& InAddress)
+	void FMessagePacket::SetArguments(TArray<FOSCData> OSCData)
 	{
-		Address = MoveTemp(InAddress);
+		Arguments = MoveTemp(OSCData);
 	}
 
 	bool FMessagePacket::IsBundle()
diff --git a/Engine/Plugins/Runtime/OSC/Source/OSC/Private/OSCMessagePacket.h b/Engine/Plugins/Runtime/OSC/Source/OSC/Private/OSCMessagePacket.h
--- a/Engine/Plugins/Runtime/OSC/Source/OSC/Private/OSCMessagePacket.h
+++ b/Engine/Plugins/Runtime/OSC/Source/OSC/Private/OSCMessagePacket.h
@@ -19,12 +19,6 @@ namespace UE::OSC
 		/** Get endpoint IP address and port responsible for creation/forwarding of packet */
 		virtual const FIPv4Endpoint& GetIPEndpoint() const override;
 
-		UE_DEPRECATED(5.5, "Packets now have address and port set using endpoints.")
-		virtual const FString& GetIPAddress() const override { static const FString None; return None; }
-
-		UE_DEPRECATED(5.5, "Packets now have address and port set using endpoints.")
-		virtual uint16 GetPort() const override { return 0; };
-
 	protected:
 		FIPv4Endpoint IPEndpoint;
 	};
@@ -38,14 +32,16 @@ namespace UE::OSC
 		virtual ~FMessagePacket() = default;
 
 		/** Adds argument to argument array */
-		void AddArgument(FOSCData&& InOSCType);
+		void AddArgument(FOSCData OSCData);
 
 		/** Empties all arguments */
 		void EmptyArguments();
 
 		/** Set OSC message address. */
-		void SetAddress(const FOSCAddress& InAddress);
-		void SetAddress(FOSCAddress&& InAddress);
+		void SetAddress(FOSCAddress InAddress);
+
+		/** Sets argument array to the given values */
+		void SetArguments(TArray<FOSCData> Args);
 
 		/** Get OSC message address. */
 		virtual const FOSCAddress& GetAddress() const;
diff --git a/Engine/Plugins/Runtime/OSC/Source/OSC/Public/OSCClient.h b/Engine/Plugins/Runtime/OSC/Source/OSC/Public/OSCClient.h
--- a/Engine/Plugins/Runtime/OSC/Source/OSC/Public/OSCClient.h
+++ b/Engine/Plugins/Runtime/OSC/Source/OSC/Public/OSCClient.h
@@ -32,8 +32,8 @@ namespace UE::OSC
 
 		virtual bool IsActive() const = 0;
 
-		virtual void SendMessage(FOSCMessage& Message) = 0;
-		virtual void SendBundle(FOSCBundle& Bundle) = 0;
+		virtual void SendMessage(const FOSCMessage& Message) = 0;
+		virtual void SendBundle(const FOSCBundle& Bundle) = 0;
 
 		virtual void Stop() = 0;
 	};
diff --git a/Engine/Plugins/Runtime/OSC/Source/OSC/Public/OSCManager.h b/Engine/Plugins/Runtime/OSC/Source/OSC/Public/OSCManager.h
--- a/Engine/Plugins/Runtime/OSC/Source/OSC/Public/OSCManager.h
+++ b/Engine/Plugins/Runtime/OSC/Source/OSC/Public/OSCManager.h
@@ -1,7 +1,6 @@
 // Copyright Epic Games, Inc. All Rights Reserved.
 #pragma once
 
-#include "CoreMinimal.h"
 #include "Kismet/BlueprintFunctionLibrary.h"
 #include "OSCAddress.h"
 #include "OSCBundle.h"
diff --git a/Engine/Plugins/Runtime/OSC/Source/OSC/Public/OSCMessage.h b/Engine/Plugins/Runtime/OSC/Source/OSC/Public/OSCMessage.h
--- a/Engine/Plugins/Runtime/OSC/Source/OSC/Public/OSCMessage.h
+++ b/Engine/Plugins/Runtime/OSC/Source/OSC/Public/OSCMessage.h
@@ -1,10 +1,10 @@
 // Copyright Epic Games, Inc. All Rights Reserved.
 #pragma once
 
-#include "CoreMinimal.h"
 #include "OSCLog.h"
 #include "OSCPacket.h"
 #include "OSCStream.h"
+#include "OSCTypes.h"
 #include "UObject/ObjectMacros.h"
 #include "UObject/Object.h"
 #include "UObject/WeakObjectPtr.h"
@@ -18,6 +18,7 @@ struct OSC_API FOSCMessage
 	GENERATED_USTRUCT_BODY()
 
 	FOSCMessage();
+	FOSCMessage(FOSCAddress Address, TArray<UE::OSC::FOSCData> Args);
 	FOSCMessage(const TSharedRef<UE::OSC::IPacket>& InPacket);
 
 	UE_DEPRECATED(5.5, "Use shared ref ctor instead")
diff --git a/Engine/Plugins/Runtime/OSC/Source/OSC/Public/OSCPacket.h b/Engine/Plugins/Runtime/OSC/Source/OSC/Public/OSCPacket.h
--- a/Engine/Plugins/Runtime/OSC/Source/OSC/Public/OSCPacket.h
+++ b/Engine/Plugins/Runtime/OSC/Source/OSC/Public/OSCPacket.h
@@ -31,11 +31,19 @@ namespace UE::OSC
 		/** Returns true if packet is bundle */
 		virtual bool IsBundle() = 0;
 
-		UE_DEPRECATED(5.5, "Packets now have address and port set using endpoints.")
-		virtual const FString& GetIPAddress() const = 0;
-
-		UE_DEPRECATED(5.5, "Packets now have port and port set using endpoints.")
-		virtual uint16 GetPort() const = 0;
+		UE_DEPRECATED(5.5, "Packet port can now be accessed via GetIPEndpoint() call.")
+		virtual const FString& GetIPAddress() const
+		{
+			static FString Addr;
+			Addr = GetIPEndpoint().ToString();
+			return Addr;
+		}
+
+		UE_DEPRECATED(5.5, "Packet address can now be accessed via GetIPEndpoint() call.")
+		virtual uint16 GetPort() const
+		{
+			return GetIPEndpoint().Port;
+		};
 
 		virtual const FIPv4Endpoint& GetIPEndpoint() const = 0;
 
