Add socket errors to warnings & display when osc server proxy successfully connects
#rnx

diff --git a/Engine/Plugins/Runtime/OSC/Source/OSC/Private/OSCServerReceiver.cpp b/Engine/Plugins/Runtime/OSC/Source/OSC/Private/OSCServerReceiver.cpp
--- a/Engine/Plugins/Runtime/OSC/Source/OSC/Private/OSCServerReceiver.cpp
+++ b/Engine/Plugins/Runtime/OSC/Source/OSC/Private/OSCServerReceiver.cpp
@@ -29,7 +29,7 @@ namespace UE::OSC
 		}
 	} // namespace ServerReceieverPrivate
 
-	FServerReceiver::FServerReceiver(FPrivateToken, FOptions&& InOptions)
+	FServerReceiver::FServerReceiver(FPrivateToken, FOptions InOptions)
 		: Options(MoveTemp(InOptions))
 	{
 	}
@@ -69,8 +69,8 @@ namespace UE::OSC
 		{
 			const ESocketErrors SocketError = SocketSubsystem.GetLastErrorCode();
 
-			UE_LOG(LogOSC, Warning, TEXT("OSC Socket Receiver: Failed to configure socket %s: Error code %d."),
-				*InName, int32(SocketError));
+			UE_LOG(LogOSC, Warning, TEXT("OSC Socket Receiver: Failed to configure socket %s: Error code %d: %s"),
+				*InName, int32(SocketError), SocketSubsystem.GetSocketError(SocketError));
 
 			return false;
 		}
@@ -79,8 +79,8 @@ namespace UE::OSC
 		{
 			const ESocketErrors SocketError = SocketSubsystem.GetLastErrorCode();
 
-			UE_LOG(LogOSC, Warning, TEXT("OSC Socket Receiver: Failed to bind %s to %s. Error code %d."),
-				*InName, *InEndpoint.ToString(), int32(SocketError));
+			UE_LOG(LogOSC, Warning, TEXT("OSC Socket Receiver: Failed to bind %s to %s. Error code %d: %s"),
+				*InName, *InEndpoint.ToString(), int32(SocketError), SocketSubsystem.GetSocketError(SocketError));
 
 			return false;
 		}
@@ -165,7 +165,7 @@ namespace UE::OSC
 		return InvalidReceiverName;
 	}
 
-	TSharedRef<FServerReceiver> FServerReceiver::Launch(const FString& InName, const FIPv4Endpoint& InEndpoint, FOptions&& InOptions)
+	TSharedRef<FServerReceiver> FServerReceiver::Launch(const FString& InName, const FIPv4Endpoint& InEndpoint, FOptions InOptions)
 	{
 		TSharedRef<FServerReceiver> NewReceiver = MakeShared<FServerReceiver>(FPrivateToken { }, MoveTemp(InOptions));
 
@@ -175,6 +175,7 @@ namespace UE::OSC
 			if (NewReceiver->Socket)
 			{
 				NewReceiver->StartThread(InName, InEndpoint);
+				UE_LOG(LogOSC, Display, TEXT("OSCServer '%s' started"), *InName);
 			}
 		}
 		else
diff --git a/Engine/Plugins/Runtime/OSC/Source/OSC/Private/OSCServerReceiver.h b/Engine/Plugins/Runtime/OSC/Source/OSC/Private/OSCServerReceiver.h
--- a/Engine/Plugins/Runtime/OSC/Source/OSC/Private/OSCServerReceiver.h
+++ b/Engine/Plugins/Runtime/OSC/Source/OSC/Private/OSCServerReceiver.h
@@ -42,10 +42,10 @@ namespace UE::OSC
 			uint32 MaxReadBufferSize = 65507u;
 		};
 
-		explicit FServerReceiver(FPrivateToken, FOptions&& InOptions);
+		explicit FServerReceiver(FPrivateToken, FOptions InOptions);
 		virtual ~FServerReceiver();
 
-		static TSharedRef<FServerReceiver> Launch(const FString& InName, const FIPv4Endpoint& InEndpoint, FOptions&& InOptions);
+		static TSharedRef<FServerReceiver> Launch(const FString& InName, const FIPv4Endpoint& InEndpoint, FOptions InOptions);
 
 		FString GetDescription() const;
 
