Fix for hardcoded byte size of 4 instead of using size of template numeric type when writing numeric type to OSC stream

diff --git a/Engine/Plugins/Runtime/OSC/Source/OSC/Public/OSCStream.h b/Engine/Plugins/Runtime/OSC/Source/OSC/Public/OSCStream.h
--- a/Engine/Plugins/Runtime/OSC/Source/OSC/Public/OSCStream.h
+++ b/Engine/Plugins/Runtime/OSC/Source/OSC/Public/OSCStream.h
@@ -134,7 +134,7 @@ namespace UE::OSC
 			return *(TNumericType*)(Temp) = InValue;
 #endif // !PLATFORM_LITTLE_ENDIAN
 
-			Write(Temp, 4);
+			Write(Temp, NumBytes);
 		}
 
 		/** Read data from stream into a provided buffer. */
