Fix for `FOSCAddress::PopContainer` to return the popped value, and for `FOSCAddress::PopContainers` to pop the right number of values.

The current version of `PopContainer` always returns an empty string, and `PopContainers` pops one more container than requested.

Test for `PopContainer` that now works:
```
FOSCAddress Address;
FString RootContainer = TEXT("Root");
FString Container = TEXT("Container");
FString Method = TEXT("Method");

Address.PushContainer(RootContainer);
Address.PushContainer(Container);
Address.SetMethod(Method);
check(Address.GetContainerPath() == FString::Printf(TEXT("/%s/%s"), *RootContainer, *Container));
check(Address.GetFullPath() == FString::Printf(TEXT("/%s/%s/%s"), *RootContainer, *Container, *Method));

FString PoppedContainer;
Address.PopContainer(PoppedContainer);
check(PoppedContainer == Container);
check(Address.GetContainerPath() == FString::Printf(TEXT("/%s"), *RootContainer));
check(Address.GetFullPath() == FString::Printf(TEXT("/%s/%s"), *RootContainer, *Method));

Address.PopContainer(PoppedContainer);
check(PoppedContainer == RootContainer);
check(Address.GetContainerPath() == TEXT("/"));
check(Address.GetFullPath() == FString::Printf(TEXT("/%s"), *Method));

Address.PopContainer(PoppedContainer);
check(PoppedContainer.IsEmpty());
```

Test for `PopContainers` that now works:
```
FOSCAddress Address;
FString RootContainer = TEXT("Root");
FString Container = TEXT("Container");
FString Method = TEXT("Method");

Address.PushContainer(RootContainer);
Address.PushContainer(Container);
Address.SetMethod(Method);

TArray<FString> Containers = Address.PopContainers(1);
check(Containers.Num() == 1);
check(Containers[0] == Container);
check(Address.GetContainerPath() == FString::Printf(TEXT("/%s"), *RootContainer));
check(Address.GetFullPath() == FString::Printf(TEXT("/%s/%s"), *RootContainer, *Method));
```

diff --git a/Engine/Plugins/Runtime/OSC/Source/OSC/Private/OSCAddress.cpp b/Engine/Plugins/Runtime/OSC/Source/OSC/Private/OSCAddress.cpp
--- a/Engine/Plugins/Runtime/OSC/Source/OSC/Private/OSCAddress.cpp
+++ b/Engine/Plugins/Runtime/OSC/Source/OSC/Private/OSCAddress.cpp
@@ -98,16 +98,14 @@ FString FOSCAddress::PopContainer()
 
 bool FOSCAddress::PopContainer(FString& OutContainer)
 {
-	OutContainer = { };
-
-	FString Popped;
 	if (Containers.IsEmpty())
 	{
+		OutContainer = { };
 		return false;
 	}
 
-	Popped = Containers.Pop(EAllowShrinking::No);
-	Hash = GetTypeHash(GetFullPath());
+	OutContainer = Containers.Pop(EAllowShrinking::No);
+	CacheAggregates();
 	return true;
 }
 
@@ -127,14 +125,12 @@ bool FOSCAddress::PopContainers(int32 NumToPop, TArray<FString>& OutContainers)
 		return false;
 	}
 
-	int32 Removed = 0;
-	for (int32 Index = Containers.Num() - 1; Index >= 0 && Removed <= NumToPop; --Index)
+	while (!Containers.IsEmpty() && OutContainers.Num() != NumToPop)
 	{
-		++Removed;
 		OutContainers.Add(Containers.Pop(EAllowShrinking::No));
 	}
 
-	Hash = GetTypeHash(GetFullPath());
+	CacheAggregates();
 	return true;
 }
 
@@ -142,9 +138,10 @@ bool FOSCAddress::RemoveContainers(int32 InIndex, int32 InCount)
 {
 	if (InIndex >= 0 && InCount > 0)
 	{
-		if (InIndex + InCount < Containers.Num())
+		if (InIndex + InCount <= Containers.Num())
 		{
 			Containers.RemoveAt(InIndex, InCount);
+			CacheAggregates();
 			return true;
 		}
 	}
